Sequencer: Renamed ActorAnimation to LevelSequence

This will break existing content. An upgrade path is not available, i.e. no redirectors.

diff --git a/Engine/Plugins/MovieScene/ActorAnimationEditor/ActorAnimationEditor.uplugin b/Engine/Plugins/MovieScene/ActorAnimationEditor/ActorAnimationEditor.uplugin
deleted file mode 100644
diff --git a/Engine/Plugins/MovieScene/ActorAnimationEditor/Source/ActorAnimationEditor/ActorAnimationEditor.Build.cs b/Engine/Plugins/MovieScene/ActorAnimationEditor/Source/ActorAnimationEditor/ActorAnimationEditor.Build.cs
deleted file mode 100644
diff --git a/Engine/Plugins/MovieScene/ActorAnimationEditor/Source/ActorAnimationEditor/Private/ActorAnimationEditorModule.cpp b/Engine/Plugins/MovieScene/ActorAnimationEditor/Source/ActorAnimationEditor/Private/ActorAnimationEditorModule.cpp
deleted file mode 100644
diff --git a/Engine/Plugins/MovieScene/ActorAnimationEditor/Source/ActorAnimationEditor/Private/ActorAnimationEditorPrivatePCH.h b/Engine/Plugins/MovieScene/ActorAnimationEditor/Source/ActorAnimationEditor/Private/ActorAnimationEditorPrivatePCH.h
deleted file mode 100644
diff --git a/Engine/Plugins/MovieScene/ActorAnimationEditor/Source/ActorAnimationEditor/Private/ActorAnimationEditorToolkit.cpp b/Engine/Plugins/MovieScene/ActorAnimationEditor/Source/ActorAnimationEditor/Private/ActorAnimationEditorToolkit.cpp
deleted file mode 100644
diff --git a/Engine/Plugins/MovieScene/ActorAnimationEditor/Source/ActorAnimationEditor/Private/ActorAnimationEditorToolkit.h b/Engine/Plugins/MovieScene/ActorAnimationEditor/Source/ActorAnimationEditor/Private/ActorAnimationEditorToolkit.h
deleted file mode 100644
diff --git a/Engine/Plugins/MovieScene/ActorAnimationEditor/Source/ActorAnimationEditor/Private/AssetTools/ActorAnimationActions.cpp b/Engine/Plugins/MovieScene/ActorAnimationEditor/Source/ActorAnimationEditor/Private/AssetTools/ActorAnimationActions.cpp
deleted file mode 100644
diff --git a/Engine/Plugins/MovieScene/ActorAnimationEditor/Source/ActorAnimationEditor/Private/AssetTools/ActorAnimationActions.h b/Engine/Plugins/MovieScene/ActorAnimationEditor/Source/ActorAnimationEditor/Private/AssetTools/ActorAnimationActions.h
deleted file mode 100644
diff --git a/Engine/Plugins/MovieScene/ActorAnimationEditor/Source/ActorAnimationEditor/Private/Factories/ActorAnimationFactoryNew.cpp b/Engine/Plugins/MovieScene/ActorAnimationEditor/Source/ActorAnimationEditor/Private/Factories/ActorAnimationFactoryNew.cpp
deleted file mode 100644
diff --git a/Engine/Plugins/MovieScene/ActorAnimationEditor/Source/ActorAnimationEditor/Private/Factories/ActorAnimationFactoryNew.h b/Engine/Plugins/MovieScene/ActorAnimationEditor/Source/ActorAnimationEditor/Private/Factories/ActorAnimationFactoryNew.h
deleted file mode 100644
diff --git a/Engine/Plugins/MovieScene/ActorAnimationEditor/Source/ActorAnimationEditor/Private/Styles/ActorAnimationEditorStyle.h b/Engine/Plugins/MovieScene/ActorAnimationEditor/Source/ActorAnimationEditor/Private/Styles/ActorAnimationEditorStyle.h
deleted file mode 100644
diff --git a/Engine/Plugins/MovieScene/LevelSequenceEditor/LevelSequenceEditor.uplugin b/Engine/Plugins/MovieScene/LevelSequenceEditor/LevelSequenceEditor.uplugin
new file mode 100644
--- /dev/null
+++ b/Engine/Plugins/MovieScene/LevelSequenceEditor/LevelSequenceEditor.uplugin
@@ -0,0 +1,24 @@
+{
+	"FileVersion" : 3,
+	
+	"FriendlyName" : "Level Sequence Editor",
+	"Version" : 1,
+	"VersionName" : "1.0",
+	"CreatedBy" : "Epic Games, Inc.",
+	"CreatedByURL" : "http://epicgames.com",
+	"EngineVersion" : "4.5.0",
+	"Description" : "Content Editor for LevelSequence Assets.",
+	"Category" : "Editor",
+	"EnabledByDefault" : false,
+	"IsBetaVersion": true,
+
+	"Modules" :
+	[
+		{
+			"Name" : "LevelSequenceEditor",
+			"Type" : "Editor"
+		}
+	],
+
+	"CanContainContent" : false
+}
\ No newline at end of file
diff --git a/Engine/Plugins/MovieScene/LevelSequenceEditor/Source/LevelSequenceEditor/LevelSequenceEditor.Build.cs b/Engine/Plugins/MovieScene/LevelSequenceEditor/Source/LevelSequenceEditor/LevelSequenceEditor.Build.cs
new file mode 100644
--- /dev/null
+++ b/Engine/Plugins/MovieScene/LevelSequenceEditor/Source/LevelSequenceEditor/LevelSequenceEditor.Build.cs
@@ -0,0 +1,50 @@
+// Copyright 1998-2015 Epic Games, Inc. All Rights Reserved.
+
+using UnrealBuildTool;
+
+public class LevelSequenceEditor : ModuleRules
+{
+	public LevelSequenceEditor(TargetInfo Target)
+	{
+        DynamicallyLoadedModuleNames.AddRange(
+            new string[] {
+				"AssetTools",
+				"SceneOutliner",
+			}
+        );
+        
+        PrivateDependencyModuleNames.AddRange(
+			new string[] {
+                "LevelSequence",
+				"BlueprintGraph",
+				"Core",
+				"CoreUObject",
+                "EditorStyle",
+                "Engine",
+				"LevelEditor",
+				"MovieScene",
+				"MovieSceneTracks",
+				"Sequencer",
+                "Slate",
+                "SlateCore",
+                "UnrealEd",
+			}
+		);
+
+		PrivateIncludePathModuleNames.AddRange(
+			new string[] {
+				"AssetTools",
+				"SceneOutliner",
+			}
+		);
+
+        PrivateIncludePaths.AddRange(
+            new string[] {
+				"LevelSequenceEditor/Private",
+				"LevelSequenceEditor/Private/AssetTools",
+				"LevelSequenceEditor/Private/Factories",
+				"LevelSequenceEditor/Private/Styles",
+			}
+        );
+	}
+}
diff --git a/Engine/Plugins/MovieScene/LevelSequenceEditor/Source/LevelSequenceEditor/Private/AssetTools/LevelSequenceActions.cpp b/Engine/Plugins/MovieScene/LevelSequenceEditor/Source/LevelSequenceEditor/Private/AssetTools/LevelSequenceActions.cpp
new file mode 100644
--- /dev/null
+++ b/Engine/Plugins/MovieScene/LevelSequenceEditor/Source/LevelSequenceEditor/Private/AssetTools/LevelSequenceActions.cpp
@@ -0,0 +1,96 @@
+// Copyright 1998-2015 Epic Games, Inc. All Rights Reserved.
+
+#include "LevelSequenceEditorPCH.h"
+#include "IAssetTools.h"
+#include "ISequencerModule.h"
+#include "Toolkits/IToolkit.h"
+
+
+#define LOCTEXT_NAMESPACE "AssetTypeActions"
+
+
+/* FLevelSequenceActions constructors
+ *****************************************************************************/
+
+FLevelSequenceActions::FLevelSequenceActions( const TSharedRef<ISlateStyle>& InStyle )
+	: Style(InStyle)
+{ }
+
+
+/* IAssetTypeActions interface
+ *****************************************************************************/
+
+uint32 FLevelSequenceActions::GetCategories()
+{
+	return EAssetTypeCategories::Animation;
+}
+
+
+FText FLevelSequenceActions::GetName() const
+{
+	return NSLOCTEXT("AssetTypeActions", "AssetTypeActions_LevelSequence", "Level Sequence");
+}
+
+
+UClass* FLevelSequenceActions::GetSupportedClass() const
+{
+	return ULevelSequence::StaticClass(); 
+}
+
+
+FColor FLevelSequenceActions::GetTypeColor() const
+{
+	return FColor(200, 80, 80);
+}
+
+
+void FLevelSequenceActions::OpenAssetEditor(const TArray<UObject*>& InObjects, TSharedPtr<IToolkitHost> EditWithinLevelEditor)
+{
+	UWorld* WorldContext = nullptr;
+	for (const FWorldContext& Context : GEngine->GetWorldContexts())
+	{
+		if (Context.WorldType == EWorldType::Editor)
+		{
+			WorldContext = Context.World();
+			break;
+		}
+	}
+
+	if (!ensure(WorldContext))
+	{
+		return;
+	}
+
+	EToolkitMode::Type Mode = EditWithinLevelEditor.IsValid()
+		? EToolkitMode::WorldCentric
+		: EToolkitMode::Standalone;
+
+	for (auto ObjIt = InObjects.CreateConstIterator(); ObjIt; ++ObjIt)
+	{
+		ULevelSequence* LevelSequence = Cast<ULevelSequence>(*ObjIt);
+
+		if (LevelSequence != nullptr)
+		{
+			// Legacy upgrade
+			LevelSequence->ConvertPersistentBindingsToDefault(WorldContext);
+
+			// Create an edit instance for this level sequence that can only edit the default bindings in the current world
+			ULevelSequenceInstance* Instance = NewObject<ULevelSequenceInstance>(GetTransientPackage());
+			bool bCanInstanceBindings = false;
+			Instance->Initialize(LevelSequence, WorldContext, bCanInstanceBindings);
+
+			TSharedRef<FLevelSequenceEditorToolkit> Toolkit = MakeShareable(new FLevelSequenceEditorToolkit(Style));
+			Toolkit->Initialize(Mode, EditWithinLevelEditor, Instance, true);
+		}
+	}
+}
+
+
+bool FLevelSequenceActions::ShouldForceWorldCentric()
+{
+	// @todo sequencer: Hack to force world-centric mode for Sequencer
+	return true;
+}
+
+
+#undef LOCTEXT_NAMESPACE
\ No newline at end of file
diff --git a/Engine/Plugins/MovieScene/LevelSequenceEditor/Source/LevelSequenceEditor/Private/AssetTools/LevelSequenceActions.h b/Engine/Plugins/MovieScene/LevelSequenceEditor/Source/LevelSequenceEditor/Private/AssetTools/LevelSequenceActions.h
new file mode 100644
--- /dev/null
+++ b/Engine/Plugins/MovieScene/LevelSequenceEditor/Source/LevelSequenceEditor/Private/AssetTools/LevelSequenceActions.h
@@ -0,0 +1,38 @@
+// Copyright 1998-2015 Epic Games, Inc. All Rights Reserved.
+
+#pragma once
+
+#include "AssetTypeActions_Base.h"
+
+
+/**
+ * Implements actions for ULevelSequence assets.
+ */
+class FLevelSequenceActions
+	: public FAssetTypeActions_Base
+{
+public:
+
+	/**
+	 * Creates and initializes a new instance.
+	 *
+	 * @param InStyle The style set to use for asset editor toolkits.
+	 */
+	FLevelSequenceActions(const TSharedRef<ISlateStyle>& InStyle);
+
+public:
+	
+	// IAssetTypeActions interface
+
+	virtual uint32 GetCategories() override;
+	virtual FText GetName() const override;
+	virtual UClass* GetSupportedClass() const override;
+	virtual FColor GetTypeColor() const override;
+	virtual void OpenAssetEditor(const TArray<UObject*>& InObjects, TSharedPtr<class IToolkitHost> EditWithinLevelEditor = TSharedPtr<IToolkitHost>()) override;
+	virtual bool ShouldForceWorldCentric() override;
+
+private:
+
+	/** Pointer to the style set to use for toolkits. */
+	TSharedRef<ISlateStyle> Style;
+};
diff --git a/Engine/Plugins/MovieScene/LevelSequenceEditor/Source/LevelSequenceEditor/Private/Factories/LevelSequenceFactoryNew.cpp b/Engine/Plugins/MovieScene/LevelSequenceEditor/Source/LevelSequenceEditor/Private/Factories/LevelSequenceFactoryNew.cpp
new file mode 100644
--- /dev/null
+++ b/Engine/Plugins/MovieScene/LevelSequenceEditor/Source/LevelSequenceEditor/Private/Factories/LevelSequenceFactoryNew.cpp
@@ -0,0 +1,41 @@
+// Copyright 1998-2015 Epic Games, Inc. All Rights Reserved.
+
+#include "LevelSequenceEditorPCH.h"
+#include "LevelSequence.h"
+#include "LevelSequenceFactoryNew.h"
+
+
+#define LOCTEXT_NAMESPACE "MovieSceneFactory"
+
+
+/* ULevelSequenceFactory structors
+ *****************************************************************************/
+
+ULevelSequenceFactoryNew::ULevelSequenceFactoryNew(const FObjectInitializer& ObjectInitializer)
+	: Super(ObjectInitializer)
+{
+	bCreateNew = true;
+	bEditAfterNew = true;
+	SupportedClass = ULevelSequence::StaticClass();
+}
+
+
+/* UFactory interface
+ *****************************************************************************/
+
+UObject* ULevelSequenceFactoryNew::FactoryCreateNew(UClass* Class, UObject* InParent, FName Name, EObjectFlags Flags, UObject* Context, FFeedbackContext* Warn)
+{
+	auto NewLevelSequence = NewObject<ULevelSequence>(InParent, Name, Flags|RF_Transactional);
+	NewLevelSequence->Initialize();
+
+	return NewLevelSequence;
+}
+
+
+bool ULevelSequenceFactoryNew::ShouldShowInNewMenu() const
+{
+	return true;
+}
+
+
+#undef LOCTEXT_NAMESPACE
\ No newline at end of file
diff --git a/Engine/Plugins/MovieScene/LevelSequenceEditor/Source/LevelSequenceEditor/Private/Factories/LevelSequenceFactoryNew.h b/Engine/Plugins/MovieScene/LevelSequenceEditor/Source/LevelSequenceEditor/Private/Factories/LevelSequenceFactoryNew.h
new file mode 100644
--- /dev/null
+++ b/Engine/Plugins/MovieScene/LevelSequenceEditor/Source/LevelSequenceEditor/Private/Factories/LevelSequenceFactoryNew.h
@@ -0,0 +1,24 @@
+// Copyright 1998-2015 Epic Games, Inc. All Rights Reserved.
+
+#pragma once
+
+#include "Factories/Factory.h"
+#include "LevelSequenceFactoryNew.generated.h"
+
+
+/**
+ * Implements a factory for ULevelSequence objects.
+ */
+UCLASS(hidecategories=Object)
+class ULevelSequenceFactoryNew
+	: public UFactory
+{
+	GENERATED_UCLASS_BODY()
+
+public:
+
+	// UFactory Interface
+
+	virtual UObject* FactoryCreateNew(UClass* Class, UObject* InParent, FName Name, EObjectFlags Flags, UObject* Context, FFeedbackContext* Warn) override;
+	virtual bool ShouldShowInNewMenu() const override;
+};
diff --git a/Engine/Plugins/MovieScene/LevelSequenceEditor/Source/LevelSequenceEditor/Private/LevelSequenceEditorModule.cpp b/Engine/Plugins/MovieScene/LevelSequenceEditor/Source/LevelSequenceEditor/Private/LevelSequenceEditorModule.cpp
new file mode 100644
--- /dev/null
+++ b/Engine/Plugins/MovieScene/LevelSequenceEditor/Source/LevelSequenceEditor/Private/LevelSequenceEditorModule.cpp
@@ -0,0 +1,195 @@
+// Copyright 1998-2015 Epic Games, Inc. All Rights Reserved.
+
+#include "LevelSequenceEditorPCH.h"
+#include "LevelSequenceEditorStyle.h"
+#include "ModuleInterface.h"
+#include "LevelEditor.h"
+#include "MovieSceneActor.h"
+
+#define LOCTEXT_NAMESPACE "LevelSequenceEditor"
+
+
+class FLevelSequenceExtensionCommands
+	: public TCommands<FLevelSequenceExtensionCommands>
+{
+public:
+
+	/** Default constructor. */
+	FLevelSequenceExtensionCommands()
+		: TCommands("LevelSequenceEditor", LOCTEXT("ExtensionDescription", "Extension commands specific to the level sequence editor"), NAME_None, "LevelSequenceEditorStyle")
+	{ }
+
+	/** Initialize commands */
+	virtual void RegisterCommands() override
+	{
+		UI_COMMAND(CreateNewLevelSequenceInLevel, "Add Level Sequence", "Create a new level sequence asset, and place an instance of it in this level", EUserInterfaceActionType::Button, FInputChord());
+	}
+
+	TSharedPtr<FUICommandInfo> CreateNewLevelSequenceInLevel;
+};
+
+
+/**
+ * Implements the LevelSequenceEditor module.
+ */
+class FLevelSequenceEditorModule
+	: public IModuleInterface
+{
+public:
+
+	// IModuleInterface interface
+
+	virtual void StartupModule() override
+	{
+		Style = MakeShareable(new FLevelSequenceEditorStyle());
+
+		RegisterAssetTools();
+		RegisterMenuExtensions();
+	}
+	
+	virtual void ShutdownModule() override
+	{
+		UnregisterAssetTools();
+		UnregisterMenuExtensions();
+	}
+
+protected:
+
+	/** Registers asset tool actions. */
+	void RegisterAssetTools()
+	{
+		IAssetTools& AssetTools = FModuleManager::LoadModuleChecked<FAssetToolsModule>("AssetTools").Get();
+
+		RegisterAssetTypeAction(AssetTools, MakeShareable(new FLevelSequenceActions(Style.ToSharedRef())));
+	}
+
+	/**
+	 * Registers a single asset type action.
+	 *
+	 * @param AssetTools The asset tools object to register with.
+	 * @param Action The asset type action to register.
+	 */
+	void RegisterAssetTypeAction(IAssetTools& AssetTools, TSharedRef<IAssetTypeActions> Action)
+	{
+		AssetTools.RegisterAssetTypeActions(Action);
+		RegisteredAssetTypeActions.Add(Action);
+	}
+
+	/** Unregisters asset tool actions. */
+	void UnregisterAssetTools()
+	{
+		FAssetToolsModule* AssetToolsModule = FModuleManager::GetModulePtr<FAssetToolsModule>("AssetTools");
+
+		if (AssetToolsModule != nullptr)
+		{
+			IAssetTools& AssetTools = AssetToolsModule->Get();
+
+			for (auto Action : RegisteredAssetTypeActions)
+			{
+				AssetTools.UnregisterAssetTypeActions(Action);
+			}
+		}
+	}
+
+protected:
+
+	/** Register menu extensions for the level editor toolbar. */
+	void RegisterMenuExtensions()
+	{
+		FLevelSequenceExtensionCommands::Register();
+
+		CommandList = MakeShareable(new FUICommandList);
+
+		CommandList->MapAction(FLevelSequenceExtensionCommands::Get().CreateNewLevelSequenceInLevel,
+			FExecuteAction::CreateStatic(&FLevelSequenceEditorModule::OnCreateActorInLevel)
+		);
+
+		// Create and register the level editor toolbar menu extension
+		CinematicsMenuExtender = MakeShareable(new FExtender);
+		CinematicsMenuExtender->AddMenuExtension("LevelEditorNewMatinee", EExtensionHook::First, CommandList, FMenuExtensionDelegate::CreateStatic([](FMenuBuilder& MenuBuilder){
+			MenuBuilder.AddMenuEntry(FLevelSequenceExtensionCommands::Get().CreateNewLevelSequenceInLevel);
+		}));
+
+		FLevelEditorModule& LevelEditorModule = FModuleManager::LoadModuleChecked<FLevelEditorModule>("LevelEditor");
+		LevelEditorModule.GetAllLevelEditorToolbarCinematicsMenuExtenders().Add(CinematicsMenuExtender);
+	}
+
+	/** Unregisters menu extensions for the level editor toolbar. */
+	void UnregisterMenuExtensions()
+	{
+		if (FLevelEditorModule* LevelEditorModule = FModuleManager::GetModulePtr<FLevelEditorModule>("LevelEditor"))
+		{
+			LevelEditorModule->GetAllLevelEditorToolbarCinematicsMenuExtenders().Remove(CinematicsMenuExtender);
+		}
+		CinematicsMenuExtender = nullptr;
+		CommandList = nullptr;
+
+		FLevelSequenceExtensionCommands::Unregister();
+	}
+
+	/** Callback for creating a new level sequence asset in the level. */
+	static void OnCreateActorInLevel()
+	{
+		// Create a new level sequence
+		IAssetTools& AssetTools = FModuleManager::GetModuleChecked<FAssetToolsModule>("AssetTools").Get();
+
+		UObject* NewAsset = nullptr;
+
+		// Attempt to create a new asset
+		for (TObjectIterator<UClass> It ; It ; ++It)
+		{
+			UClass* CurrentClass = *It;
+			if (CurrentClass->IsChildOf(UFactory::StaticClass()) && !(CurrentClass->HasAnyClassFlags(CLASS_Abstract)))
+			{
+				UFactory* Factory = Cast<UFactory>(CurrentClass->GetDefaultObject());
+				if (Factory->CanCreateNew() && Factory->ImportPriority >= 0 && Factory->SupportedClass == ULevelSequence::StaticClass())
+				{
+					NewAsset = AssetTools.CreateAsset(ULevelSequence::StaticClass(), Factory);
+					break;
+				}
+			}
+		}
+
+		if (!NewAsset)
+		{
+			return;
+		}
+
+		// Spawn a  actor at the origin, and either move infront of the camera or focus camera on it (depending on the viewport) and open for edit
+		UActorFactory* ActorFactory = GEditor->FindActorFactoryForActorClass( AMovieSceneActor::StaticClass() );
+		if (!ensure(ActorFactory))
+		{
+			return;
+		}
+
+		AMovieSceneActor* NewActor = CastChecked<AMovieSceneActor>(GEditor->UseActorFactory(ActorFactory, FAssetData(NewAsset), &FTransform::Identity));
+		if( GCurrentLevelEditingViewportClient->IsPerspective() )
+		{
+			GEditor->MoveActorInFrontOfCamera( *NewActor, GCurrentLevelEditingViewportClient->GetViewLocation(), GCurrentLevelEditingViewportClient->GetViewRotation().Vector() );
+		}
+		else
+		{
+			GEditor->MoveViewportCamerasToActor( *NewActor, false );
+		}
+
+		FAssetEditorManager::Get().OpenEditorForAsset(NewAsset);
+	}
+
+private:
+
+	/** The collection of registered asset type actions. */
+	TArray<TSharedRef<IAssetTypeActions>> RegisteredAssetTypeActions;
+
+	/** Holds the plug-ins style set. */
+	TSharedPtr<ISlateStyle> Style;
+
+	/** Extender for the cinematics menu */
+	TSharedPtr<FExtender> CinematicsMenuExtender;
+
+	TSharedPtr<FUICommandList> CommandList;
+};
+
+
+IMPLEMENT_MODULE(FLevelSequenceEditorModule, LevelSequenceEditor);
+
+#undef LOCTEXT_NAMESPACE
diff --git a/Engine/Plugins/MovieScene/LevelSequenceEditor/Source/LevelSequenceEditor/Private/LevelSequenceEditorPCH.h b/Engine/Plugins/MovieScene/LevelSequenceEditor/Source/LevelSequenceEditor/Private/LevelSequenceEditorPCH.h
new file mode 100644
--- /dev/null
+++ b/Engine/Plugins/MovieScene/LevelSequenceEditor/Source/LevelSequenceEditor/Private/LevelSequenceEditorPCH.h
@@ -0,0 +1,23 @@
+// Copyright 1998-2015 Epic Games, Inc. All Rights Reserved.
+
+#pragma once
+
+
+/* Private dependencies
+ *****************************************************************************/
+
+#include "Core.h"
+#include "CoreUObject.h"
+#include "LevelSequence.h"
+#include "AssetToolsModule.h"
+#include "EditorStyle.h"
+#include "Engine.h"
+#include "SlateBasics.h"
+#include "SlateStyle.h"
+
+
+/* Private includes
+ *****************************************************************************/
+
+#include "LevelSequenceActions.h"
+#include "LevelSequenceEditorToolkit.h"
diff --git a/Engine/Plugins/MovieScene/LevelSequenceEditor/Source/LevelSequenceEditor/Private/LevelSequenceEditorToolkit.cpp b/Engine/Plugins/MovieScene/LevelSequenceEditor/Source/LevelSequenceEditor/Private/LevelSequenceEditorToolkit.cpp
new file mode 100644
--- /dev/null
+++ b/Engine/Plugins/MovieScene/LevelSequenceEditor/Source/LevelSequenceEditor/Private/LevelSequenceEditorToolkit.cpp
@@ -0,0 +1,427 @@
+// Copyright 1998-2015 Epic Games, Inc. All Rights Reserved.
+
+#include "LevelSequenceEditorPCH.h"
+#include "Toolkits/IToolkitHost.h"
+#include "Editor/LevelEditor/Public/LevelEditor.h"
+#include "Editor/LevelEditor/Public/ILevelViewport.h"
+#include "Editor/WorkspaceMenuStructure/Public/WorkspaceMenuStructure.h"
+#include "Editor/WorkspaceMenuStructure/Public/WorkspaceMenuStructureModule.h"
+#include "ISequencer.h"
+#include "ISequencerModule.h"
+#include "SDockTab.h"
+#include "MovieSceneBinding.h"
+#include "MovieScene.h"
+#include "MovieSceneMaterialTrack.h"
+#include "ScopedTransaction.h"
+#include "ClassIconFinder.h"
+
+#include "SceneOutlinerModule.h"
+#include "SceneOutlinerPublicTypes.h"
+
+
+#define LOCTEXT_NAMESPACE "LevelSequenceEditor"
+
+
+/* Local constants
+ *****************************************************************************/
+
+const FName FLevelSequenceEditorToolkit::SequencerMainTabId(TEXT("Sequencer_SequencerMain"));
+
+namespace SequencerDefs
+{
+	static const FName SequencerAppIdentifier(TEXT("SequencerApp"));
+}
+
+
+FLevelSequenceEditorToolkit::FLevelSequenceEditorToolkit(const TSharedRef<ISlateStyle>& InStyle)
+	: Style(InStyle)
+{
+	// register sequencer menu extenders
+	ISequencerModule& SequencerModule = FModuleManager::Get().LoadModuleChecked<ISequencerModule>("Sequencer");
+	int32 NewIndex = SequencerModule.GetMenuExtensibilityManager()->GetExtenderDelegates().Add(
+		FAssetEditorExtender::CreateRaw(this, &FLevelSequenceEditorToolkit::HandleMenuExtensibilityGetExtender));
+	SequencerExtenderHandle = SequencerModule.GetMenuExtensibilityManager()->GetExtenderDelegates()[NewIndex].GetHandle();
+}
+
+
+FLevelSequenceEditorToolkit::~FLevelSequenceEditorToolkit()
+{
+	Sequencer->Close();
+
+	// unregister delegates
+	if (FModuleManager::Get().IsModuleLoaded(TEXT("LevelEditor")))
+	{
+		auto& LevelEditorModule = FModuleManager::LoadModuleChecked<FLevelEditorModule>(TEXT("LevelEditor"));
+		LevelEditorModule.OnMapChanged().RemoveAll(this);
+	}
+
+	// unregister sequencer menu extenders
+	ISequencerModule& SequencerModule = FModuleManager::Get().LoadModuleChecked<ISequencerModule>("Sequencer");
+	SequencerModule.GetMenuExtensibilityManager()->GetExtenderDelegates().RemoveAll([this](const FAssetEditorExtender& Extender)
+	{
+		return SequencerExtenderHandle == Extender.GetHandle();
+	});
+}
+
+
+/* FLevelSequenceEditorToolkit interface
+ *****************************************************************************/
+
+void FLevelSequenceEditorToolkit::Initialize( const EToolkitMode::Type Mode, const TSharedPtr<IToolkitHost>& InitToolkitHost, ULevelSequenceInstance* InLevelSequenceInstance, bool bEditWithinLevelEditor )
+{
+	// create tab layout
+	const TSharedRef<FTabManager::FLayout> StandaloneDefaultLayout = FTabManager::NewLayout("Standalone_LevelSequenceEditor")
+		->AddArea
+		(
+			FTabManager::NewPrimaryArea()
+				->Split
+				(
+					FTabManager::NewStack()
+						->AddTab(SequencerMainTabId, ETabState::OpenedTab)
+				)
+		);
+
+	LevelSequenceInstance = InLevelSequenceInstance;
+
+	const bool bCreateDefaultStandaloneMenu = true;
+	const bool bCreateDefaultToolbar = false;
+
+	FAssetEditorToolkit::InitAssetEditor(Mode, InitToolkitHost, SequencerDefs::SequencerAppIdentifier, StandaloneDefaultLayout, bCreateDefaultStandaloneMenu, bCreateDefaultToolbar, LevelSequenceInstance);
+
+	// initialize sequencer
+	FSequencerInitParams SequencerInitParams;
+	{
+		SequencerInitParams.RootSequence = LevelSequenceInstance;
+		SequencerInitParams.bEditWithinLevelEditor = bEditWithinLevelEditor;
+		SequencerInitParams.ToolkitHost = InitToolkitHost;
+
+		TSharedRef<FExtender> AddMenuExtender = MakeShareable(new FExtender);
+
+		AddMenuExtender->AddMenuExtension("AddTracks", EExtensionHook::Before, nullptr,
+			FMenuExtensionDelegate::CreateLambda([=](FMenuBuilder& MenuBuilder){
+
+				MenuBuilder.AddSubMenu(
+					LOCTEXT("AddActor_Label", "Actor To Sequencer"),
+					LOCTEXT("AddActor_ToolTip", "Allow sequencer to possess an actor that already exists in the current level"),
+					FNewMenuDelegate::CreateRaw(this, &FLevelSequenceEditorToolkit::AddPosessActorMenuExtensions),
+					false /*bInOpenSubMenuOnClick*/,
+					FSlateIcon("LevelSequenceEditorStyle", "LevelSequenceEditor.PossessNewActor")
+					);
+
+			})
+		);
+
+		SequencerInitParams.ViewParams.AddMenuExtender = AddMenuExtender;
+		SequencerInitParams.ViewParams.UniqueName = "LevelSequenceEditor";
+	}
+
+	Sequencer = FModuleManager::LoadModuleChecked<ISequencerModule>("Sequencer").CreateSequencer(SequencerInitParams);
+
+	if (bEditWithinLevelEditor)
+	{
+		// @todo remove when world-centric mode is added
+		FLevelEditorModule& LevelEditorModule = FModuleManager::LoadModuleChecked<FLevelEditorModule>("LevelEditor");
+
+		LevelEditorModule.AttachSequencer(Sequencer->GetSequencerWidget(), SharedThis(this));
+
+		// We need to find out when the user loads a new map, because we might need to re-create puppet actors
+		// when previewing a MovieScene
+		LevelEditorModule.OnMapChanged().AddRaw(this, &FLevelSequenceEditorToolkit::HandleMapChanged);
+	}
+}
+
+
+/* IToolkit interface
+ *****************************************************************************/
+
+FText FLevelSequenceEditorToolkit::GetBaseToolkitName() const
+{
+	return LOCTEXT("AppLabel", "Level Sequence Editor");
+}
+
+
+FName FLevelSequenceEditorToolkit::GetToolkitFName() const
+{
+	static FName SequencerName("LevelSequenceEditor");
+	return SequencerName;
+}
+
+
+FLinearColor FLevelSequenceEditorToolkit::GetWorldCentricTabColorScale() const
+{
+	return FLinearColor(0.7, 0.0f, 0.0f, 0.5f);
+}
+
+
+FString FLevelSequenceEditorToolkit::GetWorldCentricTabPrefix() const
+{
+	return LOCTEXT("WorldCentricTabPrefix", "Sequencer ").ToString();
+}
+
+
+void FLevelSequenceEditorToolkit::RegisterTabSpawners(const TSharedRef<class FTabManager>& TabManager)
+{
+	if (IsWorldCentricAssetEditor())
+	{
+		return;
+	}
+
+	WorkspaceMenuCategory = TabManager->AddLocalWorkspaceMenuCategory(LOCTEXT("WorkspaceMenu_SequencerAssetEditor", "Sequencer"));
+
+	TabManager->RegisterTabSpawner(SequencerMainTabId, FOnSpawnTab::CreateSP(this, &FLevelSequenceEditorToolkit::HandleTabManagerSpawnTab))
+		.SetDisplayName(LOCTEXT("SequencerMainTab", "Sequencer"))
+		.SetGroup(WorkspaceMenuCategory.ToSharedRef())
+		.SetIcon(FSlateIcon(Style->GetStyleSetName(), "LevelSequenceEditor.Tabs.Sequencer"));
+}
+
+
+void FLevelSequenceEditorToolkit::UnregisterTabSpawners(const TSharedRef<class FTabManager>& TabManager)
+{
+	if (!IsWorldCentricAssetEditor())
+	{
+		TabManager->UnregisterTabSpawner(SequencerMainTabId);
+	}
+
+	// @todo remove when world-centric mode is added
+	FLevelEditorModule& LevelEditorModule = FModuleManager::LoadModuleChecked<FLevelEditorModule>("LevelEditor");
+	LevelEditorModule.AttachSequencer(SNullWidget::NullWidget, nullptr);
+}
+
+
+/* FLevelSequenceEditorToolkit callbacks
+ *****************************************************************************/
+
+void FLevelSequenceEditorToolkit::HandleAddComponentActionExecute( UActorComponent* Component )
+{
+	Sequencer->GetHandleToObject( Component );
+}
+
+
+void FLevelSequenceEditorToolkit::HandleAddComponentMaterialActionExecute( UPrimitiveComponent* Component, int32 MaterialIndex )
+{
+	FGuid ObjectHandle = Sequencer->GetHandleToObject( Component );
+	for ( const FMovieSceneBinding& Binding : Sequencer->GetFocusedMovieSceneSequence()->GetMovieScene()->GetBindings() )
+	{
+		if ( Binding.GetObjectGuid() == ObjectHandle )
+		{
+			bool bHasMaterialTrack = false;
+			for ( UMovieSceneTrack* Track : Binding.GetTracks() )
+			{
+				UMovieSceneComponentMaterialTrack* MaterialTrack = Cast<UMovieSceneComponentMaterialTrack>( Track );
+				if ( MaterialTrack->GetMaterialIndex() == MaterialIndex )
+				{
+					bHasMaterialTrack = true;
+					break;
+				}
+			}
+			if ( bHasMaterialTrack == false )
+			{
+				const FScopedTransaction Transaction( LOCTEXT( "AddComponentMaterialTrack", "Add component material track" ) );
+
+				UMovieScene* FocusedMovieScene = Sequencer->GetFocusedMovieSceneSequence()->GetMovieScene();
+				FocusedMovieScene->Modify();
+				UMovieSceneComponentMaterialTrack* MaterialTrack = Cast<UMovieSceneComponentMaterialTrack>( FocusedMovieScene->AddTrack( UMovieSceneComponentMaterialTrack::StaticClass(), ObjectHandle ) );
+				MaterialTrack->Modify();
+				MaterialTrack->SetMaterialIndex( MaterialIndex );
+				
+				Sequencer->NotifyMovieSceneDataChanged();
+			}
+		}
+	}
+}
+
+void FLevelSequenceEditorToolkit::AddActorsToSequencer(AActor*const* InActors, int32 NumActors)
+{
+	const FScopedTransaction Transaction(LOCTEXT("UndoPossessingObject", "Possess Object(s) with Sequencer"));
+
+	UMovieSceneSequence* FocussedSequence = Sequencer->GetFocusedMovieSceneSequence();
+	UMovieScene* FocussedMovieScene = FocussedSequence->GetMovieScene();
+
+	FocussedSequence->Modify();
+
+	GEditor->SelectNone(true, true);
+	while (NumActors--)
+	{
+		AActor* ThisActor = *InActors;
+
+		FGuid ObjectGuid = FocussedSequence->FindObjectId(*ThisActor);
+
+		// Add this object if it hasn't already been possessed
+		if (!ObjectGuid.IsValid() || !FocussedMovieScene->FindPossessable(ObjectGuid))
+		{
+			ObjectGuid = FocussedMovieScene->AddPossessable(ThisActor->GetName(), ThisActor->GetClass());
+			FocussedSequence->BindPossessableObject(ObjectGuid, *ThisActor);
+		}
+
+		GEditor->SelectActor(ThisActor, true, true);
+
+		InActors++;
+	}
+
+	Sequencer->NotifyMovieSceneDataChanged();
+}
+
+void FLevelSequenceEditorToolkit::AddPosessActorMenuExtensions(FMenuBuilder& MenuBuilder)
+{
+	auto IsActorValidForPossession = [=](const AActor* Actor){
+		bool bCreateHandleIfMissing = false;
+		return !Sequencer->GetHandleToObject((UObject*)Actor, bCreateHandleIfMissing).IsValid();
+	};
+
+	// Set up a menu entry to add the selected actor(s) to the sequencer
+	TArray<AActor*> ActorsValidForPossession;
+	GEditor->GetSelectedActors()->GetSelectedObjects(ActorsValidForPossession);
+	ActorsValidForPossession.RemoveAll([&](AActor* In){ return !IsActorValidForPossession(In); });
+
+	FText SelectedLabel;
+	FSlateIcon ActorIcon(FEditorStyle::GetStyleSetName(), FClassIconFinder::FindIconNameForClass(AActor::StaticClass()));
+	if (ActorsValidForPossession.Num() == 1)
+	{
+		SelectedLabel = FText::Format(LOCTEXT("AddSpecificActor", "Add '{0}'"), FText::FromString(ActorsValidForPossession[0]->GetActorLabel()));
+		FName IconName = FClassIconFinder::FindIconNameForActor(ActorsValidForPossession[0]);
+		ActorIcon = FSlateIcon(FEditorStyle::GetStyleSetName(), FClassIconFinder::FindIconNameForClass(ActorsValidForPossession[0]->GetClass()));
+	}
+	else if (ActorsValidForPossession.Num() > 1)
+	{
+		SelectedLabel = FText::Format(LOCTEXT("AddSpecificActor", "Add Current Selection ({0} actors)"), FText::AsNumber(ActorsValidForPossession.Num()));
+	}
+
+	if (!SelectedLabel.IsEmpty())
+	{
+		// Copy the array into the lambda - probably not that big a deal
+		MenuBuilder.AddMenuEntry(SelectedLabel, FText(), ActorIcon, FExecuteAction::CreateLambda([=]{
+			FSlateApplication::Get().DismissAllMenus();
+			AddActorsToSequencer(ActorsValidForPossession.GetData(), ActorsValidForPossession.Num());
+		}));
+	}
+	
+	MenuBuilder.BeginSection("ChooseActorSection", LOCTEXT("ChooseActor", "Choose Actor:"));
+
+	using namespace SceneOutliner;
+
+	// Set up a menu entry to add any arbitrary actor to the sequencer
+	FInitializationOptions InitOptions;
+	{
+		InitOptions.Mode = ESceneOutlinerMode::ActorPicker;
+
+		// We hide the header row to keep the UI compact.
+		InitOptions.bShowHeaderRow = false;
+		InitOptions.bShowSearchBox = true;
+		InitOptions.bShowCreateNewFolder = false;
+		// Only want the actor label column
+		InitOptions.ColumnMap.Add(FBuiltInColumnTypes::Label(), FColumnInfo(EColumnVisibility::Visible, 0));
+
+		// Only display actors that are not possessed already
+		InitOptions.Filters->AddFilterPredicate( FActorFilterPredicate::CreateLambda( IsActorValidForPossession ) );
+	}
+
+	// actor selector to allow the user to choose an actor
+	FSceneOutlinerModule& SceneOutlinerModule = FModuleManager::LoadModuleChecked<FSceneOutlinerModule>("SceneOutliner");
+	TSharedRef< SWidget > MiniSceneOutliner =
+		SNew( SVerticalBox )
+		+SVerticalBox::Slot()
+		.AutoHeight()
+		.MaxHeight(400.0f)
+		[
+			SceneOutlinerModule.CreateSceneOutliner(
+				InitOptions,
+				FOnActorPicked::CreateLambda([=](AActor* Actor){
+					// Create a new binding for this actor
+					FSlateApplication::Get().DismissAllMenus();
+					AddActorsToSequencer(&Actor, 1);
+				})
+			)
+		];
+	MenuBuilder.AddWidget(MiniSceneOutliner, FText::GetEmpty(), true);
+	MenuBuilder.EndSection();
+}
+
+
+void FLevelSequenceEditorToolkit::HandleMapChanged(class UWorld* NewWorld, EMapChangeType MapChangeType)
+{
+	LevelSequenceInstance->SetContext(NewWorld);
+
+	Sequencer->NotifyMapChanged(NewWorld, MapChangeType);
+}
+
+
+TSharedRef<FExtender> FLevelSequenceEditorToolkit::HandleMenuExtensibilityGetExtender( const TSharedRef<FUICommandList> CommandList, const TArray<UObject*> ContextSensitiveObjects )
+{
+	TSharedRef<FExtender> AddTrackMenuExtender(new FExtender());
+	AddTrackMenuExtender->AddMenuExtension(
+		SequencerMenuExtensionPoints::AddTrackMenu_PropertiesSection,
+		EExtensionHook::Before,
+		CommandList,
+		FMenuExtensionDelegate::CreateRaw(this, &FLevelSequenceEditorToolkit::HandleTrackMenuExtensionAddTrack, ContextSensitiveObjects));
+
+	return AddTrackMenuExtender;
+}
+
+
+TSharedRef<SDockTab> FLevelSequenceEditorToolkit::HandleTabManagerSpawnTab(const FSpawnTabArgs& Args)
+{
+	TSharedPtr<SWidget> TabWidget = SNullWidget::NullWidget;
+
+	if (Args.GetTabId() == SequencerMainTabId)
+	{
+		TabWidget = Sequencer->GetSequencerWidget();
+	}
+
+	return SNew(SDockTab)
+		.Label(LOCTEXT("SequencerMainTitle", "Sequencer"))
+		.TabColorScale(GetTabColorScale())
+		.TabRole(ETabRole::PanelTab)
+		[
+			TabWidget.ToSharedRef()
+		];
+}
+
+
+void FLevelSequenceEditorToolkit::HandleTrackMenuExtensionAddTrack( FMenuBuilder& AddTrackMenuBuilder, TArray<UObject*> ContextObjects )
+{
+	if (ContextObjects.Num() != 1)
+	{
+		return;
+	}
+
+	AActor* Actor = Cast<AActor>(ContextObjects[0]);
+	if (Actor != nullptr)
+	{
+		AddTrackMenuBuilder.BeginSection("Components", LOCTEXT("ComponentsSection", "Components"));
+		{
+			for (UActorComponent* Component : Actor->GetComponents())
+			{
+				FUIAction AddComponentAction(FExecuteAction::CreateSP(this, &FLevelSequenceEditorToolkit::HandleAddComponentActionExecute, Component));
+				FText AddComponentLabel = FText::FromString(Component->GetName());
+				FText AddComponentToolTip = FText::Format(LOCTEXT("ComponentToolTipFormat", "Add {0} component"), FText::FromString(Component->GetName()));
+				AddTrackMenuBuilder.AddMenuEntry(AddComponentLabel, AddComponentToolTip, FSlateIcon(), AddComponentAction);
+			}
+		}
+		AddTrackMenuBuilder.EndSection();
+	}
+	else
+	{
+		UPrimitiveComponent* Component = Cast<UPrimitiveComponent>(ContextObjects[0]);
+		if (Component != nullptr)
+		{
+			int32 NumMaterials = Component->GetNumMaterials();
+			if (NumMaterials > 0)
+			{
+				AddTrackMenuBuilder.BeginSection("Materials", LOCTEXT("MaterialSection", "Materials"));
+				{
+					for (int32 MaterialIndex = 0; MaterialIndex < NumMaterials; MaterialIndex++)
+					{
+						FUIAction AddComponentMaterialAction(FExecuteAction::CreateSP(this, &FLevelSequenceEditorToolkit::HandleAddComponentMaterialActionExecute, Component, MaterialIndex));
+						FText AddComponentMaterialLabel = FText::Format(LOCTEXT("ComponentMaterialIndexLabelFormat", "Element {0}"), FText::AsNumber(MaterialIndex));
+						FText AddComponentMaterialToolTip = FText::Format(LOCTEXT("ComponentMaterialIndexToolTipFormat", "Add material element {0}" ), FText::AsNumber(MaterialIndex));
+						AddTrackMenuBuilder.AddMenuEntry( AddComponentMaterialLabel, AddComponentMaterialToolTip, FSlateIcon(), AddComponentMaterialAction );
+					}
+				}
+				AddTrackMenuBuilder.EndSection();
+			}
+		}
+	}
+}
+
+
+#undef LOCTEXT_NAMESPACE
diff --git a/Engine/Plugins/MovieScene/LevelSequenceEditor/Source/LevelSequenceEditor/Private/LevelSequenceEditorToolkit.h b/Engine/Plugins/MovieScene/LevelSequenceEditor/Source/LevelSequenceEditor/Private/LevelSequenceEditorToolkit.h
new file mode 100644
--- /dev/null
+++ b/Engine/Plugins/MovieScene/LevelSequenceEditor/Source/LevelSequenceEditor/Private/LevelSequenceEditorToolkit.h
@@ -0,0 +1,114 @@
+// Copyright 1998-2015 Epic Games, Inc. All Rights Reserved.
+
+#pragma once
+
+#include "Toolkits/AssetEditorToolkit.h"
+#include "LevelSequence.h"
+
+enum class EMapChangeType : uint8;
+class FTabManager;
+class ILevelViewport;
+class ISequencer;
+class IToolkitHost;
+class SWidget;
+class UWorld;
+struct FSequencerViewParams;
+
+
+/**
+ * Implements an Editor toolkit for level sequences.
+ */
+class FLevelSequenceEditorToolkit
+	: public FAssetEditorToolkit
+	, public FGCObject
+{ 
+public:
+
+	/**
+	 * Creates and initializes a new instance.
+	 *
+	 * @param InStyle The style set to use.
+	 */
+	FLevelSequenceEditorToolkit(const TSharedRef<ISlateStyle>& InStyle);
+
+	/** Virtual destructor */
+	virtual ~FLevelSequenceEditorToolkit();
+
+public:
+
+	/**
+	 * Initialize this asset editor.
+	 *
+	 * @param Mode Asset editing mode for this editor (standalone or world-centric).
+	 * @param InitToolkitHost When Mode is WorldCentric, this is the level editor instance to spawn this editor within.
+	 * @param LevelSequence The animation to edit.
+	 * @param TrackEditorDelegates Delegates to call to create auto-key handlers for this sequencer.
+	 * @param bEditWithinLevelEditor Whether or not sequencer should be edited within the level editor.
+	 */
+	void Initialize(const EToolkitMode::Type Mode, const TSharedPtr<IToolkitHost>& InitToolkitHost, ULevelSequenceInstance* LevelSequenceInstance, bool bEditWithinLevelEditor);
+
+public:
+
+	// IToolkit interface
+
+	virtual FText GetBaseToolkitName() const override;
+	virtual FName GetToolkitFName() const override;
+	virtual FLinearColor GetWorldCentricTabColorScale() const override;
+	virtual FString GetWorldCentricTabPrefix() const override;
+	virtual void RegisterTabSpawners(const TSharedRef<FTabManager>& TabManager) override;
+	virtual void UnregisterTabSpawners(const TSharedRef<FTabManager>& TabManager) override;
+
+	void UpdateViewports(AActor* ActorToViewThrough) const;
+
+	void AddReferencedObjects( FReferenceCollector& Collector ) override
+	{
+		Collector.AddReferencedObject(LevelSequenceInstance);
+	}
+
+private:
+
+	/** Callback for executing the Add Component action. */
+	void HandleAddComponentActionExecute(UActorComponent* Component);
+
+	/** Callback for executing the add component material track. */
+	void HandleAddComponentMaterialActionExecute(UPrimitiveComponent* Component, int32 MaterialIndex);
+
+	/** Callback for map changes. */
+	void HandleMapChanged(UWorld* NewWorld, EMapChangeType MapChangeType);
+
+	/** Callback for the menu extensibility manager. */
+	TSharedRef<FExtender> HandleMenuExtensibilityGetExtender(const TSharedRef<FUICommandList> CommandList, const TArray<UObject*> ContextSensitiveObjects);
+
+	/** Callback for spawning tabs. */
+	TSharedRef<SDockTab> HandleTabManagerSpawnTab(const FSpawnTabArgs& Args);	
+
+	/** Callback for the track menu extender. */
+	void HandleTrackMenuExtensionAddTrack(FMenuBuilder& AddTrackMenuBuilder, TArray<UObject*> ContextObjects);
+
+	/** Menu extension callback for the add menu */
+	void AddPosessActorMenuExtensions(FMenuBuilder& MenuBuilder);
+
+	/** Add the specified actors to the sequencer */
+	void AddActorsToSequencer(AActor*const* InActors, int32 NumActors);
+
+private:
+
+	/** Level sequence instance for our edit operation. */
+	ULevelSequenceInstance* LevelSequenceInstance;
+
+	/** The sequencer used by this editor. */
+	TSharedPtr<ISequencer> Sequencer;
+
+	/** A map of all the transport controls to viewports that this sequencer has made */
+	TMap<TWeakPtr<ILevelViewport>, TSharedPtr<SWidget>> TransportControls;
+
+	FDelegateHandle SequencerExtenderHandle;
+
+	/** Pointer to the style set to use for toolkits. */
+	TSharedRef<ISlateStyle> Style;
+
+private:
+
+	/**	The tab ids for all the tabs used */
+	static const FName SequencerMainTabId;
+};
diff --git a/Engine/Plugins/MovieScene/LevelSequenceEditor/Source/LevelSequenceEditor/Private/Styles/LevelSequenceEditorStyle.h b/Engine/Plugins/MovieScene/LevelSequenceEditor/Source/LevelSequenceEditor/Private/Styles/LevelSequenceEditorStyle.h
new file mode 100644
--- /dev/null
+++ b/Engine/Plugins/MovieScene/LevelSequenceEditor/Source/LevelSequenceEditor/Private/Styles/LevelSequenceEditorStyle.h
@@ -0,0 +1,53 @@
+// Copyright 1998-2015 Epic Games, Inc. All Rights Reserved.
+
+#pragma once
+
+
+#define IMAGE_BRUSH(RelativePath, ...) FSlateImageBrush(RootToContentDir(RelativePath, TEXT(".png")), __VA_ARGS__)
+#define BOX_BRUSH(RelativePath, ...) FSlateBoxBrush(RootToContentDir(RelativePath, TEXT(".png")), __VA_ARGS__)
+#define BORDER_BRUSH(RelativePath, ...) FSlateBorderBrush(RootToContentDir(RelativePath, TEXT(".png")), __VA_ARGS__)
+#define TTF_FONT(RelativePath, ...) FSlateFontInfo(RootToContentDir(RelativePath, TEXT(".ttf")), __VA_ARGS__)
+#define OTF_FONT(RelativePath, ...) FSlateFontInfo(RootToContentDir(RelativePath, TEXT(".otf")), __VA_ARGS__)
+
+
+/**
+ * Implements the visual style of the messaging debugger UI.
+ */
+class FLevelSequenceEditorStyle
+	: public FSlateStyleSet
+{
+public:
+
+	/** Default constructor. */
+	 FLevelSequenceEditorStyle()
+		 : FSlateStyleSet("LevelSequenceEditorStyle")
+	 {
+		const FVector2D Icon16x16(16.0f, 16.0f);
+		const FVector2D Icon20x20(20.0f, 20.0f);
+		const FVector2D Icon40x40(40.0f, 40.0f);
+
+		SetContentRoot(FPaths::EnginePluginsDir() / TEXT("MovieScene/LevelSequenceEditor/Content"));
+
+		// tab icons
+		Set("LevelSequenceEditor.Tabs.Sequencer", new IMAGE_BRUSH("icon_tab_sequencer_16x", Icon16x16));
+		Set("LevelSequenceEditor.PossessNewActor", new IMAGE_BRUSH("PossessNewActor_16x", Icon16x16));
+		Set("LevelSequenceEditor.PossessNewActor.Small", new IMAGE_BRUSH("PossessNewActor_16x", Icon16x16));
+		Set("LevelSequenceEditor.CreateNewLevelSequenceInLevel", new IMAGE_BRUSH("CreateNewLevelSequenceInLevel_16x", Icon16x16));
+		Set("LevelSequenceEditor.CreateNewLevelSequenceInLevel.Small", new IMAGE_BRUSH("CreateNewLevelSequenceInLevel_16x", Icon16x16));
+
+		FSlateStyleRegistry::RegisterSlateStyle(*this);
+	 }
+
+	 /** Virtual destructor. */
+	 virtual ~FLevelSequenceEditorStyle()
+	 {
+		FSlateStyleRegistry::UnRegisterSlateStyle(*this);
+	 }
+};
+
+
+#undef IMAGE_BRUSH
+#undef BOX_BRUSH
+#undef BORDER_BRUSH
+#undef TTF_FONT
+#undef OTF_FONT
diff --git a/Engine/Source/Editor/DetailCustomizations/Private/ActorAnimationPlaybackSettingsCustomization.cpp b/Engine/Source/Editor/DetailCustomizations/Private/ActorAnimationPlaybackSettingsCustomization.cpp
deleted file mode 100644
diff --git a/Engine/Source/Editor/DetailCustomizations/Private/ActorAnimationPlaybackSettingsCustomization.h b/Engine/Source/Editor/DetailCustomizations/Private/ActorAnimationPlaybackSettingsCustomization.h
deleted file mode 100644
diff --git a/Engine/Source/Editor/DetailCustomizations/Private/DetailCustomizations.cpp b/Engine/Source/Editor/DetailCustomizations/Private/DetailCustomizations.cpp
--- a/Engine/Source/Editor/DetailCustomizations/Private/DetailCustomizations.cpp
+++ b/Engine/Source/Editor/DetailCustomizations/Private/DetailCustomizations.cpp
@@ -104,7 +104,7 @@
 #include "ConfigEditorPropertyDetails.h"
 #include "AssetImportDataCustomization.h"
 #include "CaptureResolutionCustomization.h"
-#include "ActorAnimationPlaybackSettingsCustomization.h"
+#include "LevelSequencePlaybackSettingsCustomization.h"
 #include "TextCustomization.h"
 #include "AnimTrailNodeDetails.h"
 
@@ -213,7 +213,7 @@ void FDetailCustomizationsModule::RegisterPropertyTypeCustomizations()
 	RegisterCustomPropertyTypeLayout("HierarchicalSimplification", FOnGetPropertyTypeCustomizationInstance::CreateStatic(&FHierarchicalSimplificationCustomizations::MakeInstance));
 	RegisterCustomPropertyTypeLayout("PostProcessSettings", FOnGetPropertyTypeCustomizationInstance::CreateStatic(&FPostProcessSettingsCustomization::MakeInstance));
 	RegisterCustomPropertyTypeLayout("AssetImportInfo", FOnGetPropertyTypeCustomizationInstance::CreateStatic(&FAssetImportDataCustomization::MakeInstance));
-	RegisterCustomPropertyTypeLayout("ActorAnimationPlaybackSettings", FOnGetPropertyTypeCustomizationInstance::CreateStatic(&FActorAnimationPlaybackSettingsCustomization::MakeInstance));
+	RegisterCustomPropertyTypeLayout("LevelSequencePlaybackSettings", FOnGetPropertyTypeCustomizationInstance::CreateStatic(&FLevelSequencePlaybackSettingsCustomization::MakeInstance));
 	RegisterCustomPropertyTypeLayout("CaptureResolution", FOnGetPropertyTypeCustomizationInstance::CreateStatic(&FCaptureResolutionCustomization::MakeInstance));
 	RegisterCustomPropertyTypeLayout("WeightedBlendable", FOnGetPropertyTypeCustomizationInstance::CreateStatic(&FWeightedBlendableCustomization::MakeInstance));
 }
diff --git a/Engine/Source/Editor/DetailCustomizations/Private/LevelSequencePlaybackSettingsCustomization.cpp b/Engine/Source/Editor/DetailCustomizations/Private/LevelSequencePlaybackSettingsCustomization.cpp
new file mode 100644
--- /dev/null
+++ b/Engine/Source/Editor/DetailCustomizations/Private/LevelSequencePlaybackSettingsCustomization.cpp
@@ -0,0 +1,130 @@
+﻿// Copyright 1998-2015 Epic Games, Inc. All Rights Reserved.
+
+#include "DetailCustomizationsPrivatePCH.h"
+#include "LevelSequencePlaybackSettingsCustomization.h"
+#include "LevelSequencePlayer.h"
+
+#define LOCTEXT_NAMESPACE "LevelSequencePlaybackSettingsCustomization"
+
+TSharedRef<IPropertyTypeCustomization> FLevelSequencePlaybackSettingsCustomization::MakeInstance()
+{
+	return MakeShareable( new FLevelSequencePlaybackSettingsCustomization );
+}
+
+void FLevelSequencePlaybackSettingsCustomization::CustomizeHeader( TSharedRef<IPropertyHandle> InPropertyHandle, FDetailWidgetRow& HeaderRow, IPropertyTypeCustomizationUtils& CustomizationUtils )
+{
+
+}
+
+void FLevelSequencePlaybackSettingsCustomization::CustomizeChildren( TSharedRef<IPropertyHandle> InPropertyHandle, IDetailChildrenBuilder& ChildBuilder, IPropertyTypeCustomizationUtils& CustomizationUtils )
+{
+	StructPropertyHandle = InPropertyHandle;
+
+	ChildBuilder.AddChildProperty(StructPropertyHandle->GetChildHandle(GET_MEMBER_NAME_CHECKED(FLevelSequencePlaybackSettings, PlayRate)).ToSharedRef());
+
+
+	LoopCountProperty = StructPropertyHandle->GetChildHandle(GET_MEMBER_NAME_CHECKED(FLevelSequencePlaybackSettings, LoopCount));
+
+	// Set up the initial environment
+	{
+		LoopModes.Add(MakeShareable( new FLoopMode{ LOCTEXT("DontLoop", "Don't Loop"), 0 } ));
+		LoopModes.Add(MakeShareable( new FLoopMode{ LOCTEXT("Indefinitely", "Loop Indefinitely"), -1 }) );
+		LoopModes.Add(MakeShareable( new FLoopMode{ LOCTEXT("Exactly", "Loop Exactly..."), 1 }) );
+
+		int32 CurrentValue = -1;
+		LoopCountProperty->GetValue(CurrentValue);
+		if (auto* FoundMode = LoopModes.FindByPredicate([&](const TSharedPtr<FLoopMode>& In){ return In->Value == CurrentValue; }) )
+		{
+			CurrentMode = *FoundMode;
+		}
+		else
+		{
+			CurrentMode = LoopModes.Last();
+		}
+	}
+
+	ChildBuilder.AddChildContent(LOCTEXT("LoopTitle", "Loop"))
+	.NameContent()
+	[
+		LoopCountProperty->CreatePropertyNameWidget()
+	]
+	.ValueContent()
+	.HAlign(HAlign_Fill)
+	.MaxDesiredWidth(200)
+	[
+		SNew(SHorizontalBox)
+
+		+ SHorizontalBox::Slot()
+		.VAlign(VAlign_Center)
+		.AutoWidth()
+		[
+			SNew(SComboBox<TSharedPtr<FLoopMode>>)
+			.OptionsSource(&LoopModes)
+			.OnSelectionChanged_Lambda([&](TSharedPtr<FLoopMode> Mode, ESelectInfo::Type){
+				CurrentMode = Mode;
+				UpdateProperty();
+			})
+			.OnGenerateWidget_Lambda([&](TSharedPtr<FLoopMode> InMode){
+				return SNew(STextBlock)
+				.Font(IDetailLayoutBuilder::GetDetailFont())
+				.Text(InMode->DisplayName);
+			})
+			.InitiallySelectedItem(CurrentMode)
+			[
+				SAssignNew(CurrentText, STextBlock)
+				.Font(IDetailLayoutBuilder::GetDetailFont())
+				.Text(CurrentMode->DisplayName)
+			]
+		]
+
+		+ SHorizontalBox::Slot()
+		.Padding(FMargin(4.f,0))
+		.VAlign(VAlign_Center)
+		[
+			SAssignNew(LoopEntry, SHorizontalBox)
+			.Visibility(CurrentMode == LoopModes.Last() ? EVisibility::Visible : EVisibility::Collapsed)
+
+			+ SHorizontalBox::Slot()
+			.VAlign(VAlign_Center)
+			.Padding(FMargin(0, 0, 4.f, 0))
+			[
+				LoopCountProperty->CreatePropertyValueWidget()
+			]
+
+			+ SHorizontalBox::Slot()
+			.VAlign(VAlign_Center)
+			.AutoWidth()
+			[
+				SNew(STextBlock)
+				.Font(IDetailLayoutBuilder::GetDetailFont())
+				.Text_Raw(this, &FLevelSequencePlaybackSettingsCustomization::GetCustomSuffix)
+			]
+		]
+	];
+}
+
+FText FLevelSequencePlaybackSettingsCustomization::GetCustomSuffix() const
+{
+	int32 NumLoops = 0;
+	LoopCountProperty->GetValue(NumLoops);
+	return NumLoops == 1 ? LOCTEXT("Time", "time") : LOCTEXT("Times", "times");
+}
+
+void FLevelSequencePlaybackSettingsCustomization::UpdateProperty()
+{
+	if (CurrentMode == LoopModes.Last())
+	{
+		// Show custom loop entry
+		LoopEntry->SetVisibility(EVisibility::Visible);
+	}
+	else
+	{
+		// Hide custom loop entry
+		LoopEntry->SetVisibility(EVisibility::Collapsed);
+	}
+
+	LoopCountProperty->SetValue(CurrentMode->Value);
+	CurrentText->SetText(CurrentMode->DisplayName);
+}
+
+#undef LOCTEXT_NAMESPACE
diff --git a/Engine/Source/Editor/DetailCustomizations/Private/LevelSequencePlaybackSettingsCustomization.h b/Engine/Source/Editor/DetailCustomizations/Private/LevelSequencePlaybackSettingsCustomization.h
new file mode 100644
--- /dev/null
+++ b/Engine/Source/Editor/DetailCustomizations/Private/LevelSequencePlaybackSettingsCustomization.h
@@ -0,0 +1,50 @@
+// Copyright 1998-2015 Epic Games, Inc. All Rights Reserved.
+
+#pragma once
+
+
+class FLevelSequencePlaybackSettingsCustomization
+	: public IPropertyTypeCustomization
+{
+public:
+
+	/** Makes a new instance of this detail layout class for a specific detail view requesting it. */
+	static TSharedRef<IPropertyTypeCustomization> MakeInstance();
+
+public:
+
+	// IDetailCustomization interface
+
+	virtual void CustomizeHeader( TSharedRef<IPropertyHandle> PropertyHandle, FDetailWidgetRow& HeaderRow, IPropertyTypeCustomizationUtils& CustomizationUtils ) override;
+	virtual void CustomizeChildren( TSharedRef<IPropertyHandle> PropertyHandle, IDetailChildrenBuilder& ChildBuilder, IPropertyTypeCustomizationUtils& CustomizationUtils ) override;
+
+protected:
+
+	/** Respond to a selection change event from the combo box. */
+	void UpdateProperty();
+	
+private:
+
+	/** Get the suffix to display after the custom entry box. */
+	FText GetCustomSuffix() const;
+
+	/** Array of loop modes. */
+	struct FLoopMode
+	{
+		FText DisplayName;
+		int32 Value;
+	};
+
+	TArray<TSharedPtr<FLoopMode>> LoopModes;
+
+	/** The loop mode we're currently displaying. */
+	TSharedPtr<FLoopMode> CurrentMode;
+	/** The text of the current selection. */
+	TSharedPtr<STextBlock> CurrentText;
+	/** The loop number entry to be hidden and shown based on combo box selection. */
+	TSharedPtr<SWidget> LoopEntry;
+
+	/** Property handles of the properties we're editing */
+	TSharedPtr<IPropertyHandle> StructPropertyHandle;
+	TSharedPtr<IPropertyHandle> LoopCountProperty;
+};
diff --git a/Engine/Source/Editor/EditorStyle/Private/SlateEditorStyle.cpp b/Engine/Source/Editor/EditorStyle/Private/SlateEditorStyle.cpp
--- a/Engine/Source/Editor/EditorStyle/Private/SlateEditorStyle.cpp
+++ b/Engine/Source/Editor/EditorStyle/Private/SlateEditorStyle.cpp
@@ -5286,7 +5286,7 @@ void FSlateEditorStyle::FStyle::SetupClassIconsAndThumbnails()
 
 			TEXT("Actor"),
 			TEXT("ActorComponent"),
-			TEXT("ActorAnimation"),
+			TEXT("LevelSequence"),
 			TEXT("AIController"),
 			TEXT("AimOffsetBlendSpace"),
 			TEXT("AimOffsetBlendSpace1D"),
diff --git a/Engine/Source/Editor/LevelEditor/LevelEditor.Build.cs b/Engine/Source/Editor/LevelEditor/LevelEditor.Build.cs
--- a/Engine/Source/Editor/LevelEditor/LevelEditor.Build.cs
+++ b/Engine/Source/Editor/LevelEditor/LevelEditor.Build.cs
@@ -31,7 +31,7 @@ public class LevelEditor : ModuleRules
 
 		PrivateDependencyModuleNames.AddRange(
 			new string[] {
-				"ActorAnimation",
+				"LevelSequence",
 				"Analytics",
 				"Core",
 				"CoreUObject",
diff --git a/Engine/Source/Editor/LevelEditor/Private/LevelEditorToolBar.cpp b/Engine/Source/Editor/LevelEditor/Private/LevelEditorToolBar.cpp
--- a/Engine/Source/Editor/LevelEditor/Private/LevelEditorToolBar.cpp
+++ b/Engine/Source/Editor/LevelEditor/Private/LevelEditorToolBar.cpp
@@ -1233,7 +1233,7 @@ TSharedRef< SWidget > FLevelEditorToolBar::MakeLevelEditorToolBar( const TShared
 			FUIAction(),
 			FOnGetContent::CreateStatic( &FLevelEditorToolBar::GenerateCinematicsMenuContent, InCommandList, TWeakPtr<SLevelEditor>( InLevelEditor ) ),
 			LOCTEXT( "EditCinematics_Label", "Cinematics" ),
-			LOCTEXT( "EditCinematics_Tooltip", "Displays a list of Matinee and Actor Animation objects to open in their respective editors"),
+			LOCTEXT( "EditCinematics_Tooltip", "Displays a list of Matinee and Level Sequence objects to open in their respective editors"),
 			FSlateIcon(FEditorStyle::GetStyleSetName(), "LevelEditor.EditMatinee") 
 			);
 	}
@@ -2123,17 +2123,18 @@ void FLevelEditorToolBar::OnCinematicsActorPicked( AActor* Actor )
 	FSlateApplication::Get().DismissAllMenus();
 
 	// Make sure we dismiss the menus before we open this
-	if(AMatineeActor* MatineeActor = Cast<AMatineeActor>(Actor))
+	if (AMatineeActor* MatineeActor = Cast<AMatineeActor>(Actor))
 	{
 		// Open Matinee for editing!
 		GEditor->OpenMatinee( MatineeActor );
 	}
 	else if (AMovieSceneActor* MovieSceneActor = Cast<AMovieSceneActor>(Actor))
 	{
-		UObject* Animation = MovieSceneActor->ActorAnimation.TryLoad();
-		if (Animation)
+		UObject* Asset = MovieSceneActor->LevelSequence.TryLoad();
+
+		if (Asset != nullptr)
 		{
-			FAssetEditorManager::Get().OpenEditorForAsset(Animation);
+			FAssetEditorManager::Get().OpenEditorForAsset(Asset);
 		}
 	}
 }
diff --git a/Engine/Source/Editor/LevelEditor/Private/SLevelEditor.cpp b/Engine/Source/Editor/LevelEditor/Private/SLevelEditor.cpp
--- a/Engine/Source/Editor/LevelEditor/Private/SLevelEditor.cpp
+++ b/Engine/Source/Editor/LevelEditor/Private/SLevelEditor.cpp
@@ -752,11 +752,11 @@ TSharedRef<SDockTab> SLevelEditor::SpawnLevelEditorTab( const FSpawnTabArgs& Arg
 	}
 	else if( TabIdentifier == TEXT("Sequencer") )
 	{
-		if (FSlateStyleRegistry::FindSlateStyle("ActorAnimationEditorStyle"))
+		if (FSlateStyleRegistry::FindSlateStyle("LevelSequenceEditorStyle"))
 		{
 			// @todo sequencer: remove when world-centric mode is added
 			SequencerTab = SNew(SDockTab)
-				.Icon( FSlateStyleRegistry::FindSlateStyle("ActorAnimationEditorStyle")->GetBrush("ActorAnimationEditor.Tabs.Sequencer") )
+				.Icon( FSlateStyleRegistry::FindSlateStyle("LevelSequenceEditorStyle")->GetBrush("LevelSequenceEditor.Tabs.Sequencer") )
 				.Label( NSLOCTEXT("Sequencer", "SequencerMainTitle", "Sequencer") )
 				[
 					SNullWidget::NullWidget
@@ -1096,7 +1096,7 @@ TSharedRef<SWidget> SLevelEditor::RestoreContentArea( const TSharedRef<SDockTab>
 
 		{
 			// @todo remove when world-centric mode is added
-			const FSlateIcon SequencerIcon("ActorAnimationEditorStyle", "ActorAnimationEditor.Tabs.Sequencer" );
+			const FSlateIcon SequencerIcon("LevelSequenceEditorStyle", "LevelSequenceEditor.Tabs.Sequencer" );
 			LevelEditorTabManager->RegisterTabSpawner( "Sequencer", FOnSpawnTab::CreateSP<SLevelEditor, FName, FString>(this, &SLevelEditor::SpawnLevelEditorTab, FName("Sequencer"), FString()) )
 				.SetDisplayName(NSLOCTEXT("LevelEditorTabs", "Sequencer", "Sequencer"))
 				.SetGroup( MenuStructure.GetLevelEditorCategory() )
diff --git a/Engine/Source/Editor/MovieSceneTools/Private/TrackEditors/EventTrackEditor.cpp b/Engine/Source/Editor/MovieSceneTools/Private/TrackEditors/EventTrackEditor.cpp
--- a/Engine/Source/Editor/MovieSceneTools/Private/TrackEditors/EventTrackEditor.cpp
+++ b/Engine/Source/Editor/MovieSceneTools/Private/TrackEditors/EventTrackEditor.cpp
@@ -39,7 +39,7 @@ void FEventTrackEditor::BuildAddTrackMenu(FMenuBuilder& MenuBuilder)
 {
 	UMovieSceneSequence* RootMovieSceneSequence = GetSequencer()->GetRootMovieSceneSequence();
 
-	if ((RootMovieSceneSequence == nullptr) || (RootMovieSceneSequence->GetClass()->GetName() != TEXT("ActorAnimationInstance")))
+	if ((RootMovieSceneSequence == nullptr) || (RootMovieSceneSequence->GetClass()->GetName() != TEXT("LevelSequenceInstance")))
 	{
 		return;
 	}
diff --git a/Engine/Source/Editor/MovieSceneTools/Private/TrackEditors/FadeTrackEditor.cpp b/Engine/Source/Editor/MovieSceneTools/Private/TrackEditors/FadeTrackEditor.cpp
--- a/Engine/Source/Editor/MovieSceneTools/Private/TrackEditors/FadeTrackEditor.cpp
+++ b/Engine/Source/Editor/MovieSceneTools/Private/TrackEditors/FadeTrackEditor.cpp
@@ -33,7 +33,7 @@ void FFadeTrackEditor::BuildAddTrackMenu(FMenuBuilder& MenuBuilder)
 {
 	UMovieSceneSequence* RootMovieSceneSequence = GetSequencer()->GetRootMovieSceneSequence();
 
-	if ((RootMovieSceneSequence == nullptr) || (RootMovieSceneSequence->GetClass()->GetName() != TEXT("ActorAnimationInstance")))
+	if ((RootMovieSceneSequence == nullptr) || (RootMovieSceneSequence->GetClass()->GetName() != TEXT("LevelSequenceInstance")))
 	{
 		return;
 	}
diff --git a/Engine/Source/Editor/MovieSceneTools/Private/TrackEditors/ShotTrackEditor.cpp b/Engine/Source/Editor/MovieSceneTools/Private/TrackEditors/ShotTrackEditor.cpp
--- a/Engine/Source/Editor/MovieSceneTools/Private/TrackEditors/ShotTrackEditor.cpp
+++ b/Engine/Source/Editor/MovieSceneTools/Private/TrackEditors/ShotTrackEditor.cpp
@@ -562,7 +562,7 @@ void FShotTrackEditor::BuildAddTrackMenu(FMenuBuilder& MenuBuilder)
 {
 	UMovieSceneSequence* RootMovieSceneSequence = GetSequencer()->GetRootMovieSceneSequence();
 
-	if ((RootMovieSceneSequence == nullptr) || (RootMovieSceneSequence->GetClass()->GetName() != TEXT("ActorAnimationInstance")))
+	if ((RootMovieSceneSequence == nullptr) || (RootMovieSceneSequence->GetClass()->GetName() != TEXT("LevelSequenceInstance")))
 	{
 		return;
 	}
diff --git a/Engine/Source/Editor/MovieSceneTools/Private/TrackEditors/SlomoTrackEditor.cpp b/Engine/Source/Editor/MovieSceneTools/Private/TrackEditors/SlomoTrackEditor.cpp
--- a/Engine/Source/Editor/MovieSceneTools/Private/TrackEditors/SlomoTrackEditor.cpp
+++ b/Engine/Source/Editor/MovieSceneTools/Private/TrackEditors/SlomoTrackEditor.cpp
@@ -33,7 +33,7 @@ void FSlomoTrackEditor::BuildAddTrackMenu(FMenuBuilder& MenuBuilder)
 {
 	UMovieSceneSequence* RootMovieSceneSequence = GetSequencer()->GetRootMovieSceneSequence();
 
-	if ((RootMovieSceneSequence == nullptr) || (RootMovieSceneSequence->GetClass()->GetName() != TEXT("ActorAnimationInstance")))
+	if ((RootMovieSceneSequence == nullptr) || (RootMovieSceneSequence->GetClass()->GetName() != TEXT("LevelSequenceInstance")))
 	{
 		return;
 	}
diff --git a/Engine/Source/Editor/Sequencer/Private/Sequencer.cpp b/Engine/Source/Editor/Sequencer/Private/Sequencer.cpp
--- a/Engine/Source/Editor/Sequencer/Private/Sequencer.cpp
+++ b/Engine/Source/Editor/Sequencer/Private/Sequencer.cpp
@@ -43,7 +43,7 @@
 #include "IMenu.h"
 
 #include "MovieSceneCaptureModule.h"
-#include "AutomatedActorAnimationCapture.h"
+#include "AutomatedLevelSequenceCapture.h"
 
 #define LOCTEXT_NAMESPACE "Sequencer"
 
@@ -2350,12 +2350,12 @@ void FSequencer::BindSequencerCommands()
 			FLevelEditorModule& LevelEditorModule = FModuleManager::GetModuleChecked<FLevelEditorModule>(TEXT("LevelEditor"));
 			IMovieSceneCaptureModule& MovieSceneCaptureModule = IMovieSceneCaptureModule::Get();
 
-			// Create a new movie scene capture object for an automated actor animation, and open the tab
-			UAutomatedActorAnimationCapture* MovieSceneCapture = NewObject<UAutomatedActorAnimationCapture>(GetTransientPackage(), UAutomatedActorAnimationCapture::StaticClass(), NAME_None, RF_Transient);
+			// Create a new movie scene capture object for an automated level sequence, and open the tab
+			UAutomatedLevelSequenceCapture* MovieSceneCapture = NewObject<UAutomatedLevelSequenceCapture>(GetTransientPackage(), UAutomatedLevelSequenceCapture::StaticClass(), NAME_None, RF_Transient);
 			MovieSceneCapture->LoadConfig();
 
-			// Set the actor animation asset we want to render
-			MovieSceneCapture->ActorAnimation = GetCurrentAsset()->GetPathName();
+			// Set the level sequence asset we want to render
+			MovieSceneCapture->LevelSequence = GetCurrentAsset()->GetPathName();
 
 			// Set the level we want to render
 			for (const FWorldContext& Context : GEngine->GetWorldContexts())
@@ -2366,6 +2366,7 @@ void FSequencer::BindSequencerCommands()
 					break;
 				}
 			}
+
 			MovieSceneCaptureModule.OpenCaptureSettings(LevelEditorModule.GetLevelEditorTabManager().ToSharedRef(), MovieSceneCapture);
 		})
 	);
diff --git a/Engine/Source/Editor/Sequencer/Private/Sequencer.h b/Engine/Source/Editor/Sequencer/Private/Sequencer.h
--- a/Engine/Source/Editor/Sequencer/Private/Sequencer.h
+++ b/Engine/Source/Editor/Sequencer/Private/Sequencer.h
@@ -376,7 +376,7 @@ protected:
 	}
 
 	/** 
-	 * Find the viewed actor animation asset in the content browser
+	 * Find the viewed sequence asset in the content browser
 	 */
 	void FindInContentBrowser();
 
diff --git a/Engine/Source/Editor/Sequencer/Private/SequencerCommands.cpp b/Engine/Source/Editor/Sequencer/Private/SequencerCommands.cpp
--- a/Engine/Source/Editor/Sequencer/Private/SequencerCommands.cpp
+++ b/Engine/Source/Editor/Sequencer/Private/SequencerCommands.cpp
@@ -54,7 +54,7 @@ void FSequencerCommands::RegisterCommands()
 
 	UI_COMMAND( ToggleSnapCurveValueToInterval, "Snap Curve Key Values", "Snap curve keys to the value snapping interval", EUserInterfaceActionType::ToggleButton, FInputChord() );
 
-	UI_COMMAND( FindInContentBrowser, "Find in Content Browser", "Find the viewed actor animation in the content browser", EUserInterfaceActionType::Button, FInputChord() );
+	UI_COMMAND( FindInContentBrowser, "Find in Content Browser", "Find the viewed sequence asset in the content browser", EUserInterfaceActionType::Button, FInputChord() );
 	UI_COMMAND( ToggleDetailsView, "Details View", "Enable details view for selected sections and keys", EUserInterfaceActionType::ToggleButton, FInputChord() );
 	
 	UI_COMMAND( ToggleShowCurveEditor, "Curve Editor", "Show the animation keys in a curve editor", EUserInterfaceActionType::ToggleButton, FInputChord() );
diff --git a/Engine/Source/Editor/Sequencer/Private/SequencerCommands.h b/Engine/Source/Editor/Sequencer/Private/SequencerCommands.h
--- a/Engine/Source/Editor/Sequencer/Private/SequencerCommands.h
+++ b/Engine/Source/Editor/Sequencer/Private/SequencerCommands.h
@@ -120,7 +120,7 @@ public:
 	/** Toggles whether or not to snap curve values to the interval. */
 	TSharedPtr< FUICommandInfo > ToggleSnapCurveValueToInterval;
 
-	/** Finds the viewed actor animation in the content browser. */
+	/** Finds the viewed sequence asset in the content browser. */
 	TSharedPtr< FUICommandInfo > FindInContentBrowser;
 
 	/** Toggles whether the details view is enabled in the level editor. */
diff --git a/Engine/Source/Editor/Sequencer/Public/ISequencerModule.h b/Engine/Source/Editor/Sequencer/Public/ISequencerModule.h
--- a/Engine/Source/Editor/Sequencer/Public/ISequencerModule.h
+++ b/Engine/Source/Editor/Sequencer/Public/ISequencerModule.h
@@ -11,7 +11,7 @@
 class ISequencerObjectBindingManager;
 class ISequencerTrackEditor;
 class IToolkitHost;
-class UActorAnimation;
+class ULevelSequence;
 
 
 namespace SequencerMenuExtensionPoints
diff --git a/Engine/Source/Editor/UnrealEd/Private/Factories/ActorFactory.cpp b/Engine/Source/Editor/UnrealEd/Private/Factories/ActorFactory.cpp
--- a/Engine/Source/Editor/UnrealEd/Private/Factories/ActorFactory.cpp
+++ b/Engine/Source/Editor/UnrealEd/Private/Factories/ActorFactory.cpp
@@ -67,7 +67,7 @@ ActorFactory.cpp:
 #include "Kismet2/ComponentEditorUtils.h"
 #include "Components/BillboardComponent.h"
 
-#include "ActorAnimation.h"
+#include "LevelSequence.h"
 #include "MovieSceneActor.h"
 #include "ActorFactoryMovieScene.h"
 
@@ -1960,9 +1960,9 @@ bool UActorFactoryMovieScene::CanCreateActorFrom( const FAssetData& AssetData, F
 		return true;
 	}
 
-	if ( AssetData.IsValid() && !AssetData.GetClass()->IsChildOf( UActorAnimation::StaticClass() ) )
+	if ( AssetData.IsValid() && !AssetData.GetClass()->IsChildOf( ULevelSequence::StaticClass() ) )
 	{
-		OutErrorMsg = NSLOCTEXT("CanCreateActor", "NoActorAnimationAsset", "A valid actor animation asset must be specified.");
+		OutErrorMsg = NSLOCTEXT("CanCreateActor", "NoLevelSequenceAsset", "A valid sequencer asset must be specified.");
 		return false;
 	}
 
@@ -1975,9 +1975,9 @@ AActor* UActorFactoryMovieScene::SpawnActor( UObject* Asset, ULevel* InLevel, co
 
 	if (NewActor)
 	{
-		if (UActorAnimation* ActorAnimation = Cast<UActorAnimation>(Asset))
+		if (ULevelSequence* LevelSequence = Cast<ULevelSequence>(Asset))
 		{
-			NewActor->SetAnimation(ActorAnimation);
+			NewActor->SetAnimation(LevelSequence);
 		}
 	}
 
@@ -1988,7 +1988,7 @@ UObject* UActorFactoryMovieScene::GetAssetFromActorInstance(AActor* Instance)
 {
 	if (AMovieSceneActor* MovieSceneActor = Cast<AMovieSceneActor>(Instance))
 	{
-		return MovieSceneActor->ActorAnimation.TryLoad();
+		return MovieSceneActor->LevelSequence.TryLoad();
 	}
 
 	return nullptr;
diff --git a/Engine/Source/Editor/UnrealEd/UnrealEd.Build.cs b/Engine/Source/Editor/UnrealEd/UnrealEd.Build.cs
--- a/Engine/Source/Editor/UnrealEd/UnrealEd.Build.cs
+++ b/Engine/Source/Editor/UnrealEd/UnrealEd.Build.cs
@@ -93,7 +93,7 @@ public class UnrealEd : ModuleRules
 		PrivateDependencyModuleNames.AddRange(
 			new string[] 
 			{ 
-                "ActorAnimation",
+                "LevelSequence",
 				"AnimGraph",
                 "AppFramework",
 				"BlueprintGraph",
diff --git a/Engine/Source/Runtime/ActorAnimation/ActorAnimation.Build.cs b/Engine/Source/Runtime/ActorAnimation/ActorAnimation.Build.cs
deleted file mode 100644
diff --git a/Engine/Source/Runtime/ActorAnimation/Private/ActorAnimation.cpp b/Engine/Source/Runtime/ActorAnimation/Private/ActorAnimation.cpp
deleted file mode 100644
diff --git a/Engine/Source/Runtime/ActorAnimation/Private/ActorAnimationModule.cpp b/Engine/Source/Runtime/ActorAnimation/Private/ActorAnimationModule.cpp
deleted file mode 100644
diff --git a/Engine/Source/Runtime/ActorAnimation/Private/ActorAnimationObject.cpp b/Engine/Source/Runtime/ActorAnimation/Private/ActorAnimationObject.cpp
deleted file mode 100644
diff --git a/Engine/Source/Runtime/ActorAnimation/Private/ActorAnimationObjectReference.cpp b/Engine/Source/Runtime/ActorAnimation/Private/ActorAnimationObjectReference.cpp
deleted file mode 100644
diff --git a/Engine/Source/Runtime/ActorAnimation/Private/ActorAnimationPlayer.cpp b/Engine/Source/Runtime/ActorAnimation/Private/ActorAnimationPlayer.cpp
deleted file mode 100644
diff --git a/Engine/Source/Runtime/ActorAnimation/Private/ActorAnimationPrivatePCH.h b/Engine/Source/Runtime/ActorAnimation/Private/ActorAnimationPrivatePCH.h
deleted file mode 100644
diff --git a/Engine/Source/Runtime/ActorAnimation/Private/MovieSceneActor.cpp b/Engine/Source/Runtime/ActorAnimation/Private/MovieSceneActor.cpp
deleted file mode 100644
diff --git a/Engine/Source/Runtime/ActorAnimation/Public/ActorAnimation.h b/Engine/Source/Runtime/ActorAnimation/Public/ActorAnimation.h
deleted file mode 100644
diff --git a/Engine/Source/Runtime/ActorAnimation/Public/ActorAnimationObject.h b/Engine/Source/Runtime/ActorAnimation/Public/ActorAnimationObject.h
deleted file mode 100644
diff --git a/Engine/Source/Runtime/ActorAnimation/Public/ActorAnimationObjectReference.h b/Engine/Source/Runtime/ActorAnimation/Public/ActorAnimationObjectReference.h
deleted file mode 100644
diff --git a/Engine/Source/Runtime/ActorAnimation/Public/ActorAnimationPlayer.h b/Engine/Source/Runtime/ActorAnimation/Public/ActorAnimationPlayer.h
deleted file mode 100644
diff --git a/Engine/Source/Runtime/ActorAnimation/Public/MovieSceneActor.h b/Engine/Source/Runtime/ActorAnimation/Public/MovieSceneActor.h
deleted file mode 100644
diff --git a/Engine/Source/Runtime/LevelSequence/LevelSequence.Build.cs b/Engine/Source/Runtime/LevelSequence/LevelSequence.Build.cs
new file mode 100644
--- /dev/null
+++ b/Engine/Source/Runtime/LevelSequence/LevelSequence.Build.cs
@@ -0,0 +1,20 @@
+// Copyright 1998-2015 Epic Games, Inc. All Rights Reserved.
+
+using UnrealBuildTool;
+
+public class LevelSequence : ModuleRules
+{
+	public LevelSequence(TargetInfo Target)
+	{
+		PrivateIncludePaths.Add("Runtime/LevelSequence/Private");
+
+		PublicDependencyModuleNames.AddRange(
+            new string[] {
+                "Core",
+                "CoreUObject",
+                "Engine",
+				"MovieScene",
+            }
+        );
+	}
+}
diff --git a/Engine/Source/Runtime/LevelSequence/Private/LevelSequence.cpp b/Engine/Source/Runtime/LevelSequence/Private/LevelSequence.cpp
new file mode 100644
--- /dev/null
+++ b/Engine/Source/Runtime/LevelSequence/Private/LevelSequence.cpp
@@ -0,0 +1,402 @@
+// Copyright 1998-2015 Epic Games, Inc. All Rights Reserved.
+
+#include "LevelSequencePCH.h"
+#include "LevelSequence.h"
+#include "LevelSequenceObject.h"
+#include "MovieScene.h"
+
+
+ULevelSequence::ULevelSequence(const FObjectInitializer& ObjectInitializer)
+	: Super(ObjectInitializer)
+	, MovieScene(nullptr)
+{
+}
+
+void ULevelSequence::Initialize()
+{
+	// @todo sequencer: gmp: fix me
+	MovieScene = NewObject<UMovieScene>(this, NAME_None, RF_Transactional);
+}
+
+void ULevelSequence::ConvertPersistentBindingsToDefault(UObject* FixupContext)
+{
+	if (PossessedObjects_DEPRECATED.Num() == 0)
+	{
+		return;
+	}
+
+	MarkPackageDirty();
+	for (auto& Pair : PossessedObjects_DEPRECATED)
+	{
+		UObject* Object = Pair.Value.GetObject();
+		if (Object)
+		{
+			FGuid ObjectId;
+			FGuid::Parse(Pair.Key, ObjectId);
+			DefaultObjectReferences.CreateBinding(ObjectId, Object, FixupContext);
+		}
+	}
+	PossessedObjects_DEPRECATED.Empty();
+}
+
+void ULevelSequenceInstance::Initialize(ULevelSequence* InSequence, UObject* InContext, bool bInCanRemapBindings)
+{
+	LevelSequence = InSequence;
+	bCanRemapBindings = bInCanRemapBindings;
+
+	if (!bCanRemapBindings)
+	{
+		RemappedObjectReferences = FLevelSequenceObjectReferenceMap();
+	}
+
+	SetContext(InContext);
+}
+
+void ULevelSequenceInstance::SetContext(UObject* NewContext)
+{
+	ResolutionContext = NewContext;
+	CachedObjectBindings.Reset();
+}
+
+bool ULevelSequenceInstance::AllowsSpawnableObjects() const
+{
+	return true;
+}
+
+void ULevelSequenceInstance::BindPossessableObject(const FGuid& ObjectId, UObject& PossessedObject)
+{
+	UObject* Context = ResolutionContext.Get();
+
+	if (!Context)
+	{
+		return;
+	}
+
+	if (!LevelSequence)
+	{
+		RemappedObjectReferences.CreateBinding(ObjectId, &PossessedObject, Context);
+	}
+	else if (bCanRemapBindings && LevelSequence->DefaultObjectReferences.HasBinding(ObjectId))
+	{
+		RemappedObjectReferences.CreateBinding(ObjectId, &PossessedObject, Context);
+	}
+	else
+	{
+		LevelSequence->DefaultObjectReferences.CreateBinding(ObjectId, &PossessedObject, Context);
+	}
+}
+
+bool ULevelSequenceInstance::CanPossessObject(UObject& Object) const
+{
+	return Object.IsA<AActor>() || Object.IsA<UActorComponent>();
+}
+
+void ULevelSequenceInstance::DestroyAllSpawnedObjects()
+{
+	const bool DestroyAll = true;
+	SpawnOrDestroyObjects(DestroyAll);
+	SpawnedActors.Empty();
+}
+
+UObject* ULevelSequenceInstance::FindObject(const FGuid& ObjectId) const
+{
+	// If it's already cached, we can just return that
+	if (auto* WeakCachedObject = CachedObjectBindings.Find(ObjectId))
+	{
+		if (auto* CachedObj = WeakCachedObject->Get())
+		{
+			return CachedObj;
+		}
+	}
+
+	// Look up the object from scratch
+	UObject* FoundObject = nullptr;
+
+	// Otherwise we need to try and find it
+	UObject* Context = ResolutionContext.Get();
+	if (!FoundObject && Context)
+	{
+		if (!FoundObject && bCanRemapBindings)
+		{
+			// Attempt to resolve the object binding through the remapped bindings
+			FoundObject = RemappedObjectReferences.ResolveBinding(ObjectId, Context);
+		}
+		
+		if (LevelSequence)
+		{
+			// Attempt to resolve the object binding through the defaults
+			FoundObject = LevelSequence->DefaultObjectReferences.ResolveBinding(ObjectId, Context);
+		}
+	}
+
+	if (!FoundObject)
+	{
+		// Maybe it's a spawnable then?
+		FoundObject = SpawnedActors.FindRef(ObjectId).Get();
+	}
+
+	if (FoundObject)
+	{
+		CachedObjectBindings.Add(ObjectId, FoundObject);
+	}
+
+	return FoundObject;
+}
+
+FGuid ULevelSequenceInstance::FindObjectId(UObject& Object) const
+{
+	for (auto& Pair : CachedObjectBindings)
+	{
+		if (Pair.Value == &Object)
+		{
+			return Pair.Key;
+		}
+	}
+
+	// search possessed objects
+	UObject* Context = ResolutionContext.Get();
+	if (Context)
+	{
+		if (bCanRemapBindings)
+		{
+			// Search instances
+			FGuid ObjectId = RemappedObjectReferences.FindBindingId(&Object, Context);
+			if (ObjectId.IsValid())
+			{
+				return ObjectId;
+			}
+		}
+
+		// Search defaults
+		if (LevelSequence)
+		{
+			FGuid ObjectId = LevelSequence->DefaultObjectReferences.FindBindingId(&Object, Context);
+			if (ObjectId.IsValid())
+			{
+				return ObjectId;
+			}
+		}
+	}
+
+	// search spawned objects
+	for (auto SpawnedActorPair : SpawnedActors)
+	{
+		if (SpawnedActorPair.Value.Get() == &Object)
+		{
+			return SpawnedActorPair.Key;
+		}
+	}
+
+	// not found
+	return FGuid();
+}
+
+UMovieScene* ULevelSequenceInstance::GetMovieScene() const
+{
+	return LevelSequence ? LevelSequence->MovieScene : nullptr;
+}
+
+UObject* ULevelSequenceInstance::GetParentObject(UObject* Object) const
+{
+	UActorComponent* Component = Cast<UActorComponent>(Object);
+
+	if (Component != nullptr)
+	{
+		return Component->GetOwner();
+	}
+
+	return nullptr;
+}
+
+void ULevelSequenceInstance::SpawnOrDestroyObjects(bool DestroyAll)
+{
+	bool bAnyLevelActorsChanged = false;
+
+	UMovieScene* MovieScene = GetMovieScene();
+	
+	// remove any proxy actors that we no longer need
+	if (MovieScene != nullptr)
+	{
+		for (auto SpawnedActorPair : SpawnedActors)
+		{
+			bool ShouldDestroyActor = true;
+
+			if (!DestroyAll)
+			{
+				// figure out if we still need this proxy actor
+				for (auto SpawnableIndex = 0; SpawnableIndex < MovieScene->GetSpawnableCount(); ++SpawnableIndex)
+				{
+					auto& Spawnable = MovieScene->GetSpawnable(SpawnableIndex);
+
+					if (Spawnable.GetGuid() == SpawnedActorPair.Key)
+					{
+						ShouldDestroyActor = false;
+						break;
+					}
+				}
+			}
+
+			if (!ShouldDestroyActor)
+			{
+				// Actor is no longer valid (probably the world was destroyed)
+				continue;
+			}
+
+			AActor* Actor = SpawnedActorPair.Value.Get();
+				
+			if (Actor == nullptr)
+			{
+				continue;
+			}
+
+			// destroy Actor
+			UWorld* World = Actor->GetWorld();
+
+			if (ensure(World != nullptr))
+			{
+				if (!bAnyLevelActorsChanged)
+				{
+					DeselectAllActors();
+				}
+
+				// Actor should never be selected in the Editor at this point (we took care of that up above)
+				ensure(!Actor->IsSelected());
+
+				const bool bNetForce = false;
+				const bool bShouldModifyLevel = false;
+				const bool bWasDestroyed = World->DestroyActor(Actor, bNetForce, bShouldModifyLevel);
+
+				if (bWasDestroyed)
+				{
+					bAnyLevelActorsChanged = true;
+					SpawnedActors.Remove(SpawnedActorPair.Key);
+				}
+				else
+				{
+					// @todo sequencer: At least one proxy actor couldn't be cleaned up!
+				}
+			}
+		}
+	}
+
+	if (!DestroyAll && (MovieScene != nullptr))
+	{
+		for (auto SpawnableIndex = 0; SpawnableIndex < MovieScene->GetSpawnableCount(); ++SpawnableIndex)
+		{
+			FMovieSceneSpawnable& Spawnable = MovieScene->GetSpawnable(SpawnableIndex);
+
+			// Do we already have a proxy for this spawnable?
+			if (SpawnedActors.Contains(Spawnable.GetGuid()))
+			{
+				continue;
+			}
+
+			UClass* GeneratedClass = Spawnable.GetClass();
+
+			if ((GeneratedClass == nullptr) || !GeneratedClass->IsChildOf(AActor::StaticClass()))
+			{
+				continue;
+			}
+
+			AActor* ActorCDO = CastChecked<AActor>(GeneratedClass->ClassDefaultObject);
+
+			const FVector SpawnLocation = ActorCDO->GetRootComponent()->RelativeLocation;
+			const FRotator SpawnRotation = ActorCDO->GetRootComponent()->RelativeRotation;
+
+			// @todo sequencer: We should probably spawn these in a specific sub-level!
+			// World->CurrentLevel = ???;
+
+			const FName ActorName = NAME_None;
+
+			// Override the object flags so that RF_Transactional is not set.  Puppet actors are never transactional
+			// @todo sequencer: These actors need to avoid any transaction history.  However, RF_Transactional can currently be set on objects on the fly!
+			const EObjectFlags ObjectFlags = RF_Transient;		// NOTE: We are omitting RF_Transactional intentionally
+
+			// @todo sequencer livecapture: Consider using SetPlayInEditorWorld() and RestoreEditorWorld() here instead
+						
+			// @todo sequencer actors: We need to make sure puppet objects aren't copied into PIE/SIE sessions!  They should be omitted from that duplication!
+
+			// Spawn the puppet actor
+			FActorSpawnParameters SpawnInfo;
+			{
+				SpawnInfo.Name = ActorName;
+				SpawnInfo.ObjectFlags = ObjectFlags;
+			}
+
+			AActor* NewActor = GWorld->SpawnActor(GeneratedClass, &SpawnLocation, &SpawnRotation, SpawnInfo);
+				
+			if (NewActor != nullptr)
+			{
+				// @todo sequencer: We're naming the actor based off of the spawnable's name.  Is that really what we want?
+				// @todo sequencer: gmp: fix me
+				//FActorLabelUtilities::SetActorLabelUnique(NewActor, Spawnable.GetName());
+				SpawnedActors.Add(Spawnable.GetGuid(), NewActor);
+			}
+			else
+			{
+				// Actor failed to spawn
+				// @todo sequencer: What should we do when this happens to one or more actors?
+			}
+		}
+	}
+}
+
+#if WITH_EDITOR
+bool ULevelSequenceInstance::TryGetObjectDisplayName(const FGuid& ObjectId, FText& OutDisplayName) const
+{
+	UObject* Object = FindObject(ObjectId);
+
+	if (Object == nullptr)
+	{
+		return false;
+	}
+
+	AActor* Actor = Cast<AActor>(Object);
+
+	if (Actor == nullptr)
+	{
+		return false;
+	}
+
+	OutDisplayName = FText::FromString(Actor->GetActorLabel());
+
+	return true;
+}
+#endif
+
+void ULevelSequenceInstance::UnbindPossessableObjects(const FGuid& ObjectId)
+{
+	if (bCanRemapBindings && RemappedObjectReferences.HasBinding(ObjectId))
+	{
+		RemappedObjectReferences.RemoveBinding(ObjectId);
+	}
+	else if (LevelSequence)
+	{
+		LevelSequence->DefaultObjectReferences.RemoveBinding(ObjectId);
+	}
+
+	CachedObjectBindings.Remove(ObjectId);
+}
+
+void ULevelSequenceInstance::DeselectAllActors()
+{
+	// @todo sequencer: gmp: fix me
+	/*
+	USelection* Selection = GEditor->GetSelectedActors();
+	{
+		Selection->BeginBatchSelectOperation();
+		Selection->MarkBatchDirty();
+
+		for (auto SpawnedActorPair : SpawnedActors)
+		{
+			AActor* SpawnedActor = SpawnedActorPair.Value.Get();
+
+			if (SpawnedActor != nullptr)
+			{
+				GEditor->GetSelectedActors()->Deselect(SpawnedActor);
+			}
+		}
+
+		Selection->EndBatchSelectOperation();
+	}*/
+}
\ No newline at end of file
diff --git a/Engine/Source/Runtime/LevelSequence/Private/LevelSequenceModule.cpp b/Engine/Source/Runtime/LevelSequence/Private/LevelSequenceModule.cpp
new file mode 100644
--- /dev/null
+++ b/Engine/Source/Runtime/LevelSequence/Private/LevelSequenceModule.cpp
@@ -0,0 +1,22 @@
+// Copyright 1998-2015 Epic Games, Inc. All Rights Reserved.
+
+#include "LevelSequencePCH.h"
+#include "ModuleInterface.h"
+
+
+/**
+ * Implements the LevelSequence module.
+ */
+class FLevelSequenceModule
+	: public IModuleInterface
+{
+public:
+
+	// IModuleInterface interface
+
+	virtual void StartupModule() override { }
+	virtual void ShutdownModule() override { }
+};
+
+
+IMPLEMENT_MODULE(FLevelSequenceModule, LevelSequence);
diff --git a/Engine/Source/Runtime/LevelSequence/Private/LevelSequenceObject.cpp b/Engine/Source/Runtime/LevelSequence/Private/LevelSequenceObject.cpp
new file mode 100644
--- /dev/null
+++ b/Engine/Source/Runtime/LevelSequence/Private/LevelSequenceObject.cpp
@@ -0,0 +1,53 @@
+// Copyright 1998-2015 Epic Games, Inc. All Rights Reserved.
+
+#include "LevelSequencePCH.h"
+#include "LevelSequenceObject.h"
+
+
+/* FSequencerPossessedObject interface
+ *****************************************************************************/
+
+UObject* FLevelSequenceObject::GetObject() const
+{
+	// get component-less object
+	if (ComponentName.IsEmpty())
+	{
+		return ObjectOrOwner.Get();
+	}
+
+	// get cached component
+	if ((CachedComponent.IsValid()) && (CachedComponent->GetName() == ComponentName))
+	{
+		return CachedComponent.Get();
+	}
+
+	CachedComponent = nullptr;
+
+	// find & cache component
+	UObject* Object = ObjectOrOwner.Get();
+
+	if (Object == nullptr)
+	{
+		return nullptr;
+	}
+
+	AActor* Owner = Cast<AActor>(Object);
+
+	if (Owner == nullptr)
+	{
+		return Object;
+	}
+
+	for (UActorComponent* ActorComponent : Owner->GetComponents())
+	{
+		if (ActorComponent->GetName() == ComponentName)
+		{
+			CachedComponent = ActorComponent;
+
+			return ActorComponent;
+		}
+	}
+
+	// component not found
+	return nullptr;
+}
diff --git a/Engine/Source/Runtime/LevelSequence/Private/LevelSequenceObjectReference.cpp b/Engine/Source/Runtime/LevelSequence/Private/LevelSequenceObjectReference.cpp
new file mode 100644
--- /dev/null
+++ b/Engine/Source/Runtime/LevelSequence/Private/LevelSequenceObjectReference.cpp
@@ -0,0 +1,109 @@
+// Copyright 1998-2015 Epic Games, Inc. All Rights Reserved.
+
+#include "LevelSequencePCH.h"
+#include "LevelSequenceObjectReference.h"
+
+FLevelSequenceObjectReference::FLevelSequenceObjectReference(UObject* InObject, UObject* InContext)
+{
+	if (InObject->IsA<AActor>() || InObject->IsA<UActorComponent>())
+	{
+		ObjectId = FLazyObjectPtr(InObject).GetUniqueID();
+		ObjectPath = InObject->GetPathName(InContext);
+	}
+}
+
+UObject* FLevelSequenceObjectReference::Resolve(UObject* InContext) const
+{
+#if WITH_EDITOR
+	int32 PIEInstanceID = InContext->GetOutermost()->PIEInstanceID;
+	FUniqueObjectGuid FixedUpId = PIEInstanceID == -1 ? ObjectId : ObjectId.FixupForPIE(PIEInstanceID);
+#else
+	FUniqueObjectGuid FixedUpId = ObjectId;
+#endif
+	
+
+	FLazyObjectPtr LazyPtr;
+	LazyPtr = FixedUpId;
+
+	if (UObject* FoundObject = LazyPtr.Get())
+	{
+		return FoundObject;
+	}
+	else if (!ObjectPath.IsEmpty())
+	{
+		return FindObject<UObject>(InContext, *ObjectPath, false);
+	}
+
+	return nullptr;
+}
+
+bool FLevelSequenceObjectReferenceMap::Serialize(FArchive& Ar)
+{
+	int32 Num = Map.Num();
+	Ar << Num;
+
+	if (Ar.IsLoading())
+	{
+		while(Num-- > 0)
+		{
+			FGuid Key;
+			Ar << Key;
+
+			FLevelSequenceObjectReference Value;
+			Ar << Value;
+
+			Map.Add(Key, Value);
+		}
+	}
+	else if (Ar.IsSaving() || Ar.IsCountingMemory() || Ar.IsObjectReferenceCollector())
+	{
+		for (auto& Pair : Map)
+		{
+			Ar << Pair.Key;
+			Ar << Pair.Value;
+		}
+	}
+	return true;
+}
+
+bool FLevelSequenceObjectReferenceMap::HasBinding(const FGuid& ObjectId) const
+{
+	if (const FLevelSequenceObjectReference* Reference = Map.Find(ObjectId))
+	{
+		return Reference->IsValid();
+	}
+	return false;
+}
+
+void FLevelSequenceObjectReferenceMap::CreateBinding(const FGuid& ObjectId, UObject* InObject, UObject* InContext)
+{
+	FLevelSequenceObjectReference NewReference(InObject, InContext);
+
+	if (ensureMsgf(NewReference.IsValid(), TEXT("Unable to generate a reference for the specified object and context")))
+	{
+		Map.FindOrAdd(ObjectId) = NewReference;
+	}
+}
+
+void FLevelSequenceObjectReferenceMap::RemoveBinding(const FGuid& ObjectId)
+{
+	Map.Remove(ObjectId);
+}
+
+UObject* FLevelSequenceObjectReferenceMap::ResolveBinding(const FGuid& ObjectId, UObject* InContext) const
+{
+	const FLevelSequenceObjectReference* Reference = Map.Find(ObjectId);
+	return Reference ? Reference->Resolve(InContext) : nullptr;
+}
+
+FGuid FLevelSequenceObjectReferenceMap::FindBindingId(UObject* InObject, UObject* InContext) const
+{
+	for (const auto& Pair : Map)
+	{
+		if (Pair.Value.Resolve(InContext) == InObject)
+		{
+			return Pair.Key;
+		}
+	}
+	return FGuid();
+}
\ No newline at end of file
diff --git a/Engine/Source/Runtime/LevelSequence/Private/LevelSequencePCH.h b/Engine/Source/Runtime/LevelSequence/Private/LevelSequencePCH.h
new file mode 100644
--- /dev/null
+++ b/Engine/Source/Runtime/LevelSequence/Private/LevelSequencePCH.h
@@ -0,0 +1,11 @@
+// Copyright 1998-2015 Epic Games, Inc. All Rights Reserved.
+
+#pragma once
+
+
+/* Private dependencies
+ *****************************************************************************/
+
+#include "Core.h"
+#include "CoreUObject.h"
+#include "Engine.h"
diff --git a/Engine/Source/Runtime/LevelSequence/Private/LevelSequencePlayer.cpp b/Engine/Source/Runtime/LevelSequence/Private/LevelSequencePlayer.cpp
new file mode 100644
--- /dev/null
+++ b/Engine/Source/Runtime/LevelSequence/Private/LevelSequencePlayer.cpp
@@ -0,0 +1,416 @@
+// Copyright 1998-2015 Epic Games, Inc. All Rights Reserved.
+
+#include "LevelSequencePCH.h"
+#include "LevelSequencePlayer.h"
+#include "MovieScene.h"
+#include "MovieSceneSequence.h"
+#include "MovieSceneSequenceInstance.h"
+
+
+struct FTickAnimationPlayers : public FTickableGameObject
+{
+	TArray<TWeakObjectPtr<ULevelSequencePlayer>> ActiveInstances;
+	
+	virtual bool IsTickable() const override
+	{
+		return ActiveInstances.Num() != 0;
+	}
+
+	virtual TStatId GetStatId() const override
+	{
+		RETURN_QUICK_DECLARE_CYCLE_STAT(FTickAnimationPlayers, STATGROUP_Tickables);
+	}
+	
+	virtual void Tick(float DeltaSeconds) override
+	{
+		for (int32 Index = 0; Index < ActiveInstances.Num();)
+		{
+			if (auto* Player = ActiveInstances[Index].Get())
+			{
+				Player->Update(DeltaSeconds);
+				++Index;
+			}
+			else
+			{
+				ActiveInstances.RemoveAt(Index, 1, false);
+			}
+		}
+	}
+};
+
+struct FAutoDestroyAnimationTicker
+{
+	FAutoDestroyAnimationTicker()
+	{
+		FCoreDelegates::OnPreExit.AddLambda([&]{
+			Impl.Reset();
+		});
+	}
+	
+	void Add(ULevelSequencePlayer* Player)
+	{
+		if (!Impl.IsValid())
+		{
+			Impl.Reset(new FTickAnimationPlayers);
+		}
+		Impl->ActiveInstances.Add(Player);
+	}
+
+	TUniquePtr<FTickAnimationPlayers> Impl;
+} GAnimationPlayerTicker;
+
+/* ULevelSequencePlayer structors
+ *****************************************************************************/
+
+ULevelSequencePlayer::ULevelSequencePlayer(const FObjectInitializer& ObjectInitializer)
+	: Super(ObjectInitializer)
+	, LevelSequenceInstance(nullptr)
+	, bIsPlaying(false)
+	, TimeCursorPosition(0.0f)
+	, CurrentNumLoops(0)
+{ }
+
+
+/* ULevelSequencePlayer interface
+ *****************************************************************************/
+
+ULevelSequencePlayer* ULevelSequencePlayer::CreateLevelSequencePlayer(UObject* WorldContextObject, ULevelSequence* LevelSequence, FLevelSequencePlaybackSettings Settings)
+{
+	if (LevelSequence == nullptr)
+	{
+		return nullptr;
+	}
+
+	UWorld* World = GEngine->GetWorldFromContextObject(WorldContextObject);
+	check(World != nullptr);
+
+	ULevelSequencePlayer* NewPlayer = NewObject<ULevelSequencePlayer>(GetTransientPackage(), NAME_None, RF_Transient);
+	check(NewPlayer != nullptr);
+
+	// Set up a new instance of the level sequence for the player
+	ULevelSequenceInstance* Instance = NewObject<ULevelSequenceInstance>(NewPlayer);
+	Instance->Initialize(LevelSequence, World, false);
+
+	NewPlayer->Initialize(Instance, World, Settings);
+
+	// Automatically tick this player
+	GAnimationPlayerTicker.Add(NewPlayer);
+
+	return NewPlayer;
+}
+
+
+bool ULevelSequencePlayer::IsPlaying() const
+{
+	return bIsPlaying;
+}
+
+
+void ULevelSequencePlayer::Pause()
+{
+	bIsPlaying = false;
+}
+
+void ULevelSequencePlayer::Stop()
+{
+	bIsPlaying = false;
+	TimeCursorPosition = PlaybackSettings.PlayRate < 0.f ? GetLength() : 0.f;
+	CurrentNumLoops = 0;
+
+	// todo: Trigger an event?
+}
+
+void ULevelSequencePlayer::Play()
+{
+	if ((LevelSequenceInstance == nullptr) || !World.IsValid())
+	{
+		return;
+	}
+
+	// @todo Sequencer playback: Should we recreate the instance every time?
+	RootMovieSceneInstance = MakeShareable(new FMovieSceneSequenceInstance(*LevelSequenceInstance));
+
+	// @odo Sequencer Should we spawn actors here?
+	SpawnActorsForMovie(RootMovieSceneInstance.ToSharedRef());
+	RootMovieSceneInstance->RefreshInstance(*this);
+
+	bIsPlaying = true;
+
+	if (RootMovieSceneInstance.IsValid())
+	{
+		RootMovieSceneInstance->Update(TimeCursorPosition, TimeCursorPosition, *this);
+	}
+}
+
+void ULevelSequencePlayer::PlayLooping(int32 NumLoops)
+{
+	PlaybackSettings.LoopCount = NumLoops;
+	Play();
+}
+
+float ULevelSequencePlayer::GetPlaybackPosition() const
+{
+	return TimeCursorPosition;
+}
+
+void ULevelSequencePlayer::SetPlaybackPosition(float NewPlaybackPosition)
+{
+	float LastTimePosition = TimeCursorPosition;
+	
+	TimeCursorPosition = NewPlaybackPosition;
+	OnCursorPositionChanged();
+
+	if (RootMovieSceneInstance.IsValid())
+	{
+		RootMovieSceneInstance->Update(TimeCursorPosition, LastTimePosition, *this);
+	}
+}
+
+float ULevelSequencePlayer::GetLength() const
+{
+	if (!LevelSequenceInstance)
+	{
+		return 0;
+	}
+
+	UMovieScene* MovieScene = LevelSequenceInstance->GetMovieScene();
+	return MovieScene ? MovieScene->GetTimeRange().Size<float>() : 0;
+}
+
+float ULevelSequencePlayer::GetPlayRate() const
+{
+	return PlaybackSettings.PlayRate;
+}
+
+void ULevelSequencePlayer::SetPlayRate(float PlayRate)
+{
+	PlaybackSettings.PlayRate = PlayRate;
+}
+
+
+void ULevelSequencePlayer::OnCursorPositionChanged()
+{
+	float Length = GetLength();
+
+	// Handle looping or stopping
+	if (TimeCursorPosition > Length || TimeCursorPosition < 0)
+	{
+		if (PlaybackSettings.LoopCount < 0 || CurrentNumLoops < PlaybackSettings.LoopCount)
+		{
+			++CurrentNumLoops;
+			const float Overplay = FMath::Fmod(TimeCursorPosition, Length);
+			TimeCursorPosition = Overplay < 0 ? Length + Overplay : Overplay;
+		}
+		else
+		{
+			// Stop playing without modifying the playback position
+			// @todo: trigger an event?
+			bIsPlaying = false;
+			CurrentNumLoops = 0;
+		}
+	}
+}
+
+/* ULevelSequencePlayer implementation
+ *****************************************************************************/
+
+void ULevelSequencePlayer::Initialize(ULevelSequenceInstance* InLevelSequenceInstance, UWorld* InWorld, const FLevelSequencePlaybackSettings& Settings)
+{
+	LevelSequenceInstance = InLevelSequenceInstance;
+
+	World = InWorld;
+	PlaybackSettings = Settings;
+
+	// Ensure everything is set up, ready for playback
+	Stop();
+}
+
+
+/* IMovieScenePlayer interface
+ *****************************************************************************/
+
+void ULevelSequencePlayer::SpawnActorsForMovie(TSharedRef<FMovieSceneSequenceInstance> MovieSceneInstance)
+{
+	UWorld* WorldPtr = World.Get();
+
+	if (WorldPtr == nullptr)
+	{
+		return;
+	}
+
+	UMovieScene* MovieScene = MovieSceneInstance->GetSequence()->GetMovieScene();
+
+	if (MovieScene == nullptr)
+	{
+		return;
+	}
+
+	TArray<FSpawnedActorInfo>* FoundSpawnedActors = InstanceToSpawnedActorMap.Find(MovieSceneInstance);
+
+	if (FoundSpawnedActors != nullptr)
+	{
+		// Remove existing spawned actors for this movie
+		DestroyActorsForMovie( MovieSceneInstance );
+	}
+
+	TArray<FSpawnedActorInfo>& SpawnedActorList = InstanceToSpawnedActorMap.Add(MovieSceneInstance, TArray<FSpawnedActorInfo>());
+
+	for (auto SpawnableIndex = 0; SpawnableIndex < MovieScene->GetSpawnableCount(); ++SpawnableIndex)
+	{
+		auto& Spawnable = MovieScene->GetSpawnable(SpawnableIndex);
+		UClass* GeneratedClass = Spawnable.GetClass();
+		
+		if ((GeneratedClass == nullptr) || !GeneratedClass->IsChildOf(AActor::StaticClass()))
+		{
+			continue;
+		}
+
+		AActor* ActorCDO = CastChecked<AActor>(GeneratedClass->ClassDefaultObject);
+		const FVector SpawnLocation = ActorCDO->GetRootComponent()->RelativeLocation;
+		const FRotator SpawnRotation = ActorCDO->GetRootComponent()->RelativeRotation;
+
+		FActorSpawnParameters SpawnInfo;
+		{
+			SpawnInfo.ObjectFlags = RF_NoFlags;
+		}
+
+		AActor* NewActor = WorldPtr->SpawnActor(GeneratedClass, &SpawnLocation, &SpawnRotation, SpawnInfo);
+
+		if (NewActor)
+		{	
+			// Actor was spawned OK!
+			FSpawnedActorInfo NewInfo;
+			{
+				NewInfo.RuntimeGuid = Spawnable.GetGuid();
+				NewInfo.SpawnedActor = NewActor;
+			}
+						
+			SpawnedActorList.Add(NewInfo);
+		}
+	}
+}
+
+
+void ULevelSequencePlayer::DestroyActorsForMovie(TSharedRef<FMovieSceneSequenceInstance> MovieSceneInstance)
+{
+	UWorld* WorldPtr = World.Get();
+
+	if (WorldPtr == nullptr)
+	{
+		return;
+	}
+
+	TArray<FSpawnedActorInfo>* SpawnedActors = InstanceToSpawnedActorMap.Find(MovieSceneInstance);
+
+	if (SpawnedActors == nullptr)
+	{
+		return;
+	}
+
+	TArray<FSpawnedActorInfo>& SpawnedActorsRef = *SpawnedActors;
+
+	for(int32 ActorIndex = 0; ActorIndex < SpawnedActors->Num(); ++ActorIndex)
+	{
+		AActor* Actor = SpawnedActorsRef[ActorIndex].SpawnedActor.Get();
+		
+		if (Actor != nullptr)
+		{
+			// @todo Sequencer figure this out.  Defaults to false.
+			const bool bNetForce = false;
+
+			// At runtime, level modification is not needed
+			const bool bShouldModifyLevel = false;
+			Actor->GetWorld()->DestroyActor(Actor, bNetForce, bShouldModifyLevel);
+		}
+	}
+
+	InstanceToSpawnedActorMap.Remove(MovieSceneInstance);
+}
+
+
+void ULevelSequencePlayer::GetRuntimeObjects(TSharedRef<FMovieSceneSequenceInstance> MovieSceneInstance, const FGuid& ObjectId, TArray<UObject*>& OutObjects) const
+{
+	// @todo sequencer runtime: Add support for individually spawning actors on demand when first requested?
+	//    This may be important to reduce the up-front hitch when spawning actors for the entire movie, but
+	//    may introduce smaller hitches during playback.  Needs experimentation.
+
+	const TArray<FSpawnedActorInfo>* SpawnedActors = InstanceToSpawnedActorMap.Find(MovieSceneInstance);
+
+	if ((SpawnedActors != nullptr) && (SpawnedActors->Num() > 0))
+	{
+		// we have a spawned actor for this handle
+		const TArray<FSpawnedActorInfo>& SpawnedActorInfoRef = *SpawnedActors;
+
+		for (int32 SpawnedActorIndex = 0; SpawnedActorIndex < SpawnedActorInfoRef.Num(); ++SpawnedActorIndex)
+		{
+			const FSpawnedActorInfo ActorInfo = SpawnedActorInfoRef[SpawnedActorIndex];
+
+			if ((ActorInfo.RuntimeGuid == ObjectId) && ActorInfo.SpawnedActor.IsValid())
+			{
+				OutObjects.Add(ActorInfo.SpawnedActor.Get());
+			}
+		}
+	}
+	else
+	{
+		// otherwise, check whether we have one or more possessed actors that are mapped to this handle
+		UObject* FoundObject = LevelSequenceInstance->FindObject(ObjectId);
+
+		if (FoundObject != nullptr)
+		{
+			OutObjects.Add(FoundObject);
+		}
+	}
+}
+
+
+void ULevelSequencePlayer::UpdateCameraCut(UObject* ObjectToViewThrough, bool bNewCameraCut) const
+{
+}
+
+void ULevelSequencePlayer::SetViewportSettings(const TMap<FViewportClient*, EMovieSceneViewportParams>& ViewportParamsMap)
+{
+}
+
+void ULevelSequencePlayer::GetViewportSettings(TMap<FViewportClient*, EMovieSceneViewportParams>& ViewportParamsMap) const
+{
+}
+
+EMovieScenePlayerStatus::Type ULevelSequencePlayer::GetPlaybackStatus() const
+{
+	return bIsPlaying ? EMovieScenePlayerStatus::Playing : EMovieScenePlayerStatus::Stopped;
+}
+
+
+void ULevelSequencePlayer::AddOrUpdateMovieSceneInstance(UMovieSceneSection& MovieSceneSection, TSharedRef<FMovieSceneSequenceInstance> InstanceToAdd)
+{
+	SpawnActorsForMovie( InstanceToAdd );
+}
+
+
+void ULevelSequencePlayer::RemoveMovieSceneInstance(UMovieSceneSection& MovieSceneSection, TSharedRef<FMovieSceneSequenceInstance> InstanceToRemove)
+{
+	const bool bDestroyAll = true;
+	DestroyActorsForMovie( InstanceToRemove );
+}
+
+TSharedRef<FMovieSceneSequenceInstance> ULevelSequencePlayer::GetRootMovieSceneSequenceInstance() const
+{
+	return RootMovieSceneInstance.ToSharedRef();
+}
+
+void ULevelSequencePlayer::Update(const float DeltaSeconds)
+{
+	float LastTimePosition = TimeCursorPosition;
+
+	if (bIsPlaying)
+	{
+		TimeCursorPosition += DeltaSeconds * PlaybackSettings.PlayRate;
+		OnCursorPositionChanged();
+	}
+
+	if(RootMovieSceneInstance.IsValid())
+	{
+		RootMovieSceneInstance->Update(TimeCursorPosition, LastTimePosition, *this);
+	}
+}
diff --git a/Engine/Source/Runtime/LevelSequence/Private/MovieSceneActor.cpp b/Engine/Source/Runtime/LevelSequence/Private/MovieSceneActor.cpp
new file mode 100644
--- /dev/null
+++ b/Engine/Source/Runtime/LevelSequence/Private/MovieSceneActor.cpp
@@ -0,0 +1,129 @@
+// Copyright 1998-2015 Epic Games, Inc. All Rights Reserved.
+
+#include "LevelSequencePCH.h"
+#include "MovieSceneActor.h"
+
+
+AMovieSceneActor::AMovieSceneActor(const FObjectInitializer& Init)
+	: Super(Init)
+{
+#if WITH_EDITORONLY_DATA
+	UBillboardComponent* SpriteComponent = CreateEditorOnlyDefaultSubobject<UBillboardComponent>(TEXT("Sprite"));
+
+	if (!IsRunningCommandlet())
+	{
+		// Structure to hold one-time initialization
+		struct FConstructorStatics
+		{
+			ConstructorHelpers::FObjectFinderOptional<UTexture2D> DecalTexture;
+			FConstructorStatics() : DecalTexture(TEXT("/Engine/EditorResources/SceneManager")) {}
+		};
+		static FConstructorStatics ConstructorStatics;
+
+		if (SpriteComponent)
+		{
+			SpriteComponent->Sprite = ConstructorStatics.DecalTexture.Get();
+			SpriteComponent->AttachParent = RootComponent;
+			SpriteComponent->bIsScreenSizeScaled = true;
+			SpriteComponent->bAbsoluteScale = true;
+			SpriteComponent->bReceivesDecals = false;
+			SpriteComponent->bHiddenInGame = true;
+			SetRootComponent(SpriteComponent);
+		}
+	}
+#endif // WITH_EDITORONLY_DATA
+
+	PrimaryActorTick.bCanEverTick = true;
+	bAutoPlay = false;
+}
+
+void AMovieSceneActor::BeginPlay()
+{
+	Super::BeginPlay();
+	InitializePlayer();
+}
+
+void AMovieSceneActor::PostInitProperties()
+{
+	Super::PostInitProperties();
+
+	if (!HasAnyFlags(RF_ClassDefaultObject))
+	{
+		// Make an instance of our asset so that we can keep hard references to actors we are using
+		SequenceInstance = NewObject<ULevelSequenceInstance>(this, "AnimationInstance");
+	}
+}
+
+void AMovieSceneActor::PostLoad()
+{
+	Super::PostLoad();
+	
+	UpdateAnimationInstance();
+}
+
+#if WITH_EDITOR
+void AMovieSceneActor::PostEditChangeProperty( struct FPropertyChangedEvent& PropertyChangedEvent)
+{
+	FName PropertyName = PropertyChangedEvent.Property ? PropertyChangedEvent.Property->GetFName() : NAME_None;
+
+	if (PropertyName == GET_MEMBER_NAME_CHECKED(AMovieSceneActor, LevelSequence))
+	{
+		UpdateAnimationInstance();
+	}
+
+	Super::PostEditChangeProperty(PropertyChangedEvent);
+}
+
+bool AMovieSceneActor::GetReferencedContentObjects(TArray<UObject*>& Objects) const
+{
+	if (UObject* Asset = SequenceInstance ? SequenceInstance->GetLevelSequence() : nullptr)
+	{
+		// @todo: Enable editing of animation instances
+		Objects.Add(Asset);
+	}
+
+	Super::GetReferencedContentObjects(Objects);
+
+	return true;
+}
+
+#endif // WITH_EDITOR
+
+void AMovieSceneActor::Tick(float DeltaSeconds)
+{
+	if (AnimationPlayer)
+	{
+		AnimationPlayer->Update(DeltaSeconds);
+	}
+}
+
+void AMovieSceneActor::SetAnimation(ULevelSequence* InAnimation)
+{
+	if (!AnimationPlayer || !AnimationPlayer->IsPlaying())
+	{
+		LevelSequence = InAnimation;
+		InitializePlayer();
+	}
+}
+
+void AMovieSceneActor::InitializePlayer()
+{
+	UpdateAnimationInstance();
+
+	if (GetWorld()->IsGameWorld())
+	{
+		AnimationPlayer = NewObject<ULevelSequencePlayer>(this, "AnimationPlayer");
+		AnimationPlayer->Initialize(SequenceInstance, GetWorld(), PlaybackSettings);
+
+		if (bAutoPlay)
+		{
+			AnimationPlayer->Play();
+		}
+	}
+}
+
+void AMovieSceneActor::UpdateAnimationInstance()
+{
+	ULevelSequence* LevelSequenceAsset = Cast<ULevelSequence>(LevelSequence.TryLoad());
+	SequenceInstance->Initialize(LevelSequenceAsset, GetWorld(), true);
+}
\ No newline at end of file
diff --git a/Engine/Source/Runtime/LevelSequence/Public/LevelSequence.h b/Engine/Source/Runtime/LevelSequence/Public/LevelSequence.h
new file mode 100644
--- /dev/null
+++ b/Engine/Source/Runtime/LevelSequence/Public/LevelSequence.h
@@ -0,0 +1,135 @@
+// Copyright 1998-2015 Epic Games, Inc. All Rights Reserved.
+
+#pragma once
+
+#include "MovieSceneSequence.h"
+#include "LevelSequenceObject.h"
+#include "LevelSequenceObjectReference.h"
+#include "LevelSequence.generated.h"
+
+
+/**
+ * Movie scene animation for Actors.
+ */
+UCLASS(BlueprintType)
+class LEVELSEQUENCE_API ULevelSequence
+	: public UObject
+{
+	GENERATED_UCLASS_BODY()
+
+public:
+
+	/** Collection of possessed objects. */
+	UPROPERTY()
+	FLevelSequenceObjectReferenceMap DefaultObjectReferences;
+
+	/** Pointer to the movie scene that controls this animation. */
+	UPROPERTY()
+	UMovieScene* MovieScene;
+
+public:
+
+	/** Initialize this level sequence. */
+	void Initialize();
+
+	/** Convert old-style lazy object ptrs to new-style references using the specified context */
+	void ConvertPersistentBindingsToDefault(UObject* FixupContext);
+
+private:
+
+	/** Deprecated property housing old possessed object bindings */
+	UPROPERTY()
+	TMap<FString, FLevelSequenceObject> PossessedObjects_DEPRECATED;
+};
+
+
+/**
+ * An instance of a ULevelSequence that supports resolution and remapping of object bindings.
+ */
+UCLASS()
+class LEVELSEQUENCE_API ULevelSequenceInstance
+	: public UMovieSceneSequence
+{
+public:
+
+	GENERATED_BODY()
+	
+	/**
+	 * Initialize this instance. Safe to be called multiple times. Does not reset object re-mappings unless bInCanRemapBindings is false.
+	 *
+	 * @param InSequence The level sequence to instantiate.
+	 * @param InContext The context to use for object binding resolution.
+	 * @param bInCanRemapBindings Whether this instance is permitted to remap object bindings or not.
+	 */
+	void Initialize(ULevelSequence* InSequence, UObject* InContext, bool bInCanRemapBindings);
+	
+	/**
+	 * Change this instance's context.
+	 *
+	 * @param NewContext The new context to use for object binding resolution.
+	 */
+	void SetContext(UObject* NewContext);
+
+	/**
+	 * Get the level sequence of which this is an instance
+	 *
+	 * @return this instance's parent level sequence
+	 */
+	ULevelSequence* GetLevelSequence()
+	{
+		return LevelSequence;
+	}
+
+public:
+
+	// UMovieSceneSequence interface
+
+	virtual bool AllowsSpawnableObjects() const override;
+	virtual void BindPossessableObject(const FGuid& ObjectId, UObject& PossessedObject) override;
+	virtual bool CanPossessObject(UObject& Object) const override;
+	virtual void DestroyAllSpawnedObjects() override;
+	virtual UObject* FindObject(const FGuid& ObjectId) const override;
+	virtual FGuid FindObjectId(UObject& Object) const override;
+	virtual UMovieScene* GetMovieScene() const override;
+	virtual UObject* GetParentObject(UObject* Object) const override;
+	virtual void SpawnOrDestroyObjects(bool DestroyAll) override;
+	virtual void UnbindPossessableObjects(const FGuid& ObjectId) override;
+
+#if WITH_EDITOR
+	virtual bool TryGetObjectDisplayName(const FGuid& ObjectId, FText& OutDisplayName) const override;
+#endif
+
+protected:
+
+	/**
+	 * Deselect all proxy actors in the Editor.
+	 *
+	 * @todo sequencer: remove Editor dependencies from this class
+	 */
+	void DeselectAllActors();
+
+private:
+
+	/** The level sequence we are instancing. */
+	UPROPERTY(transient)
+	ULevelSequence* LevelSequence;
+
+	/** Flag to specify whether this instance can remap bindings or not. */
+	UPROPERTY()
+	bool bCanRemapBindings;
+
+	/** Collection of re-mappings for possessed objects references. */
+	UPROPERTY()
+	FLevelSequenceObjectReferenceMap RemappedObjectReferences;
+
+private:
+
+	/** A context to use for resolving object bindings */
+	TWeakObjectPtr<UObject> ResolutionContext;
+
+	/** A transient map of cached object bindings */
+	mutable TMap<FGuid, TWeakObjectPtr<UObject>> CachedObjectBindings;
+
+	/** Collection of spawned proxy actors. */
+	TMap<FGuid, TWeakObjectPtr<AActor>> SpawnedActors;
+};
diff --git a/Engine/Source/Runtime/LevelSequence/Public/LevelSequenceObject.h b/Engine/Source/Runtime/LevelSequence/Public/LevelSequenceObject.h
new file mode 100644
--- /dev/null
+++ b/Engine/Source/Runtime/LevelSequence/Public/LevelSequenceObject.h
@@ -0,0 +1,78 @@
+// Copyright 1998-2015 Epic Games, Inc. All Rights Reserved.
+
+#pragma once
+
+#include "LevelSequenceObject.generated.h"
+
+
+/**
+ * Structure for animated Actor objects.
+ */
+USTRUCT()
+struct FLevelSequenceObject
+{
+	GENERATED_BODY()
+
+public:
+
+	/** Creates and initializes a new instance. */
+	FLevelSequenceObject()
+		: ObjectOrOwner(nullptr)
+		, CachedComponent(nullptr)
+	{ }
+
+	/**
+	 * Creates and initializes a new instance from an object.
+	 *
+	 * @param InObject The object to be bound.
+	 */
+	FLevelSequenceObject(UObject* InObject)
+		: ObjectOrOwner(InObject)
+		, CachedComponent(nullptr)
+	{ }
+
+	/**
+	 * Creates and initializes a new instance from an object component.
+	 *
+	 * @param InOwner The object that owns the component.
+	 * @param InComponentName The component to be bound.
+	 */
+	FLevelSequenceObject(UObject* InOwner, FString InComponentName)
+		: ObjectOrOwner(InOwner)
+		, ComponentName(InComponentName)
+		, CachedComponent(nullptr)
+	{ }
+
+	/**
+	 * Compares two bindings for equality.
+	 *
+	 * @param X The first binding to compare.
+	 * @param Y The second binding to compare.
+	 * @return true if the bindings refer to the same object, false otherwise.
+	 */
+	friend bool operator==(const FLevelSequenceObject& X, const FLevelSequenceObject& Y)
+	{
+		return (X.ObjectOrOwner == Y.ObjectOrOwner) && (X.ComponentName == Y.ComponentName);
+	}
+
+	/**
+	 * Gets a pointer to the possessed object.
+	 *
+	 * @return The object (usually an Actor or an ActorComponent).
+	 */
+	LEVELSEQUENCE_API UObject* GetObject() const;
+
+private:
+
+	/** The object or the owner of the object being possessed. */
+	UPROPERTY()
+	TLazyObjectPtr<UObject> ObjectOrOwner;
+
+	/** Optional name of an ActorComponent. */
+	UPROPERTY()
+	FString ComponentName;
+
+	/** Cached pointer to the Actor component (only if ComponentName is set). */
+	UPROPERTY(transient)
+	mutable TWeakObjectPtr<UObject> CachedComponent;
+};
diff --git a/Engine/Source/Runtime/LevelSequence/Public/LevelSequenceObjectReference.h b/Engine/Source/Runtime/LevelSequence/Public/LevelSequenceObjectReference.h
new file mode 100644
--- /dev/null
+++ b/Engine/Source/Runtime/LevelSequence/Public/LevelSequenceObjectReference.h
@@ -0,0 +1,152 @@
+// Copyright 1998-2015 Epic Games, Inc. All Rights Reserved.
+
+#pragma once
+
+#include "LevelSequenceObjectReference.generated.h"
+
+
+/**
+ * An external reference to an level sequence object, resolvable through an arbitrary context.
+ */
+USTRUCT()
+struct FLevelSequenceObjectReference
+{
+	/**
+	 * Default construction to a null reference
+	 */
+	FLevelSequenceObjectReference() : ObjectPath() {}
+
+	/**
+	 * Generates a new reference to an object within a given context.
+	 *
+	 * @param InObject The object to generate a reference for (e.g. an AActor)
+	 * @param InContext The object to use as a context for the generation (e.g. a UWorld)
+	 */
+	LEVELSEQUENCE_API FLevelSequenceObjectReference(UObject* InObject, UObject* InContext);
+
+	GENERATED_BODY();
+
+	/**
+	 * Check whether this object reference is valid or not
+	 */
+	bool IsValid() const
+	{
+		return !ObjectPath.IsEmpty();
+	}
+
+	/**
+	 * Resolve this reference within the specified context
+	 *
+	 * @return The object (usually an Actor or an ActorComponent).
+	 */
+	LEVELSEQUENCE_API UObject* Resolve(UObject* InContext) const;
+
+	/**
+	 * Equality comparator
+	 */
+	friend bool operator==(const FLevelSequenceObjectReference& A, const FLevelSequenceObjectReference& B)
+	{
+		if (A.ObjectId.IsValid() && A.ObjectId == B.ObjectId)
+		{
+			return true;
+		}
+		return A.ObjectPath == B.ObjectPath;
+	}
+
+	/**
+	 * Serialization operator
+	 */
+	friend FArchive& operator<<(FArchive& Ar, FLevelSequenceObjectReference& Ref)
+	{
+		Ar << Ref.ObjectId;
+		Ar << Ref.ObjectPath;
+		return Ar;
+	}
+
+private:
+
+	/** Primary method of resolution - object ID, stored as an annotation on the object in the world, resolvable through TLazyObjectPtr */
+	FUniqueObjectGuid ObjectId;
+
+	/** Secondary method of resolution - path to the object within the context */
+	FString ObjectPath;
+};
+
+USTRUCT()
+struct FLevelSequenceObjectReferenceMap
+{
+	/**
+	 * Check whether this map has a binding for the specified object id
+	 * @return true if this map contains a binding for the id, false otherwise
+	 */
+	bool HasBinding(const FGuid& ObjectId) const;
+
+	/**
+	 * Remove a binding for the specified ID
+	 *
+	 * @param ObjectId	The ID to remove
+	 */
+	void RemoveBinding(const FGuid& ObjectId);
+
+	/**
+	 * Create a binding for the specified ID
+	 *
+	 * @param ObjectId	The ID to associate the object with
+	 * @param InObject	The object to associate
+	 * @param InContext	A context in which InObject resides
+	 */
+	void CreateBinding(const FGuid& ObjectId, UObject* InObject, UObject* InContext);
+
+	/**
+	 * Resolve a binding for the specified ID using a given context
+	 *
+	 * @param ObjectId	The ID to associate the object with
+	 * @param InContext	A context in which InObject resides
+	 *
+	 * @return The object if found, nullptr otherwise.
+	 */
+	UObject* ResolveBinding(const FGuid& ObjectId, UObject* InContext) const;
+
+	/**
+	 * Find an object ID that the specified object is bound with
+	 *
+	 * @param InObject	The object to find an ID for
+	 * @param InContext	The context in which InObject resides
+	 *
+	 * @return The object's bound ID, or an empty FGuid
+	 */
+	FGuid FindBindingId(UObject* InObject, UObject* InContext) const;
+
+public:
+
+	GENERATED_BODY();
+
+	/**
+	 * Equality comparator required for proper UObject serialization (so we can compare against defaults)
+	 */
+	friend bool operator==(const FLevelSequenceObjectReferenceMap& A, const FLevelSequenceObjectReferenceMap& B)
+	{
+		return A.Map.OrderIndependentCompareEqual(B.Map);
+	}
+
+	/**
+	 * Serialization function
+	 */
+	bool Serialize(FArchive& Ar);
+
+	typedef TMap<FGuid, FLevelSequenceObjectReference> MapType;
+	FORCEINLINE friend MapType::TIterator      begin(      FLevelSequenceObjectReferenceMap& Impl) { return begin(Impl.Map); }
+	FORCEINLINE friend MapType::TConstIterator begin(const FLevelSequenceObjectReferenceMap& Impl) { return begin(Impl.Map); }
+	FORCEINLINE friend MapType::TIterator      end  (      FLevelSequenceObjectReferenceMap& Impl) { return end(Impl.Map); }
+	FORCEINLINE friend MapType::TConstIterator end  (const FLevelSequenceObjectReferenceMap& Impl) { return end(Impl.Map); }
+
+private:
+	
+	// @todo sequencer: need support for UStruct keys in TMap so this can be a UPROPERTY
+	TMap<FGuid, FLevelSequenceObjectReference> Map;
+};
+
+template<> struct TStructOpsTypeTraits<FLevelSequenceObjectReferenceMap> : public TStructOpsTypeTraitsBase
+{
+	enum { WithSerializer = true, WithIdenticalViaEquality = true };
+};
\ No newline at end of file
diff --git a/Engine/Source/Runtime/LevelSequence/Public/LevelSequencePlayer.h b/Engine/Source/Runtime/LevelSequence/Public/LevelSequencePlayer.h
new file mode 100644
--- /dev/null
+++ b/Engine/Source/Runtime/LevelSequence/Public/LevelSequencePlayer.h
@@ -0,0 +1,189 @@
+// Copyright 1998-2015 Epic Games, Inc. All Rights Reserved.
+
+#pragma once
+
+#include "IMovieScenePlayer.h"
+#include "LevelSequence.h"
+#include "LevelSequencePlayer.generated.h"
+
+
+class FMovieSceneSequenceInstance;
+class ULevel;
+class UMovieSceneBindings;
+
+
+/**
+ * Settings for the level sequence player actor.
+ */
+USTRUCT(BlueprintType)
+struct FLevelSequencePlaybackSettings
+{
+	FLevelSequencePlaybackSettings()
+		: LoopCount(0)
+		, PlayRate(1.f)
+	{ }
+
+	GENERATED_BODY()
+
+	/** Number of times to loop playback. -1 for infinite, else the number of times to loop before stopping */
+	UPROPERTY(EditAnywhere, BlueprintReadWrite, Category="Playback", meta=(UIMin=1, DisplayName="Loop"))
+	int32 LoopCount;
+
+	/** The rate at which to playback the animation */
+	UPROPERTY(EditAnywhere, BlueprintReadWrite, Category="Playback", meta=(Units=Multiplier))
+	float PlayRate;
+};
+
+
+/**
+ * ULevelSequencePlayer is used to actually "play" an level sequence asset at runtime.
+ *
+ * This class keeps track of playback state and provides functions for manipulating
+ * an level sequence while its playing.
+ */
+UCLASS(BlueprintType)
+class LEVELSEQUENCE_API ULevelSequencePlayer
+	: public UObject
+	, public IMovieScenePlayer
+{
+public:
+	ULevelSequencePlayer(const FObjectInitializer&);
+
+	GENERATED_BODY()
+
+	/**
+	 * Initialize the player.
+	 *
+	 * @param InLevelSequenceInstance The level sequence instance to play.
+	 * @param InWorld The world that the animation is played in.
+	 * @param Settings The desired playback settings
+	 */
+	void Initialize(ULevelSequenceInstance* InLevelSequenceInstance, UWorld* InWorld, const FLevelSequencePlaybackSettings& Settings);
+
+public:
+
+	/**
+	 * Create a new level sequence player.
+	 *
+	 * @param WorldContextObject Context object from which to retrieve a UWorld.
+	 * @param LevelSequence The level sequence to play.
+	 * @param Settings The desired playback settings
+	 */
+	UFUNCTION(BlueprintCallable, Category="Game|Cinematic", meta=(WorldContext="WorldContextObject"))
+	static ULevelSequencePlayer* CreateLevelSequencePlayer(UObject* WorldContextObject, ULevelSequence* LevelSequence, FLevelSequencePlaybackSettings Settings);
+
+	/** Start playback from the current time cursor position, using the current play rate. */
+	UFUNCTION(BlueprintCallable, Category="Game|Cinematic")
+	void Play();
+
+	/**
+	 * Start playback from the current time cursor position, looping the specified number of times.
+	 * @param NumLoops - The number of loops to play. -1 indicates infinite looping.
+	 */
+	UFUNCTION(BlueprintCallable, Category="Game|Cinematic")
+	void PlayLooping(int32 NumLoops = -1);
+
+	/** Pause playback. */
+	UFUNCTION(BlueprintCallable, Category="Game|Cinematic")
+	void Pause();
+	
+	/** Stop playback. */
+	UFUNCTION(BlueprintCallable, Category="Game|Cinematic")
+	void Stop();
+
+	/** Get the current playback position */
+	UFUNCTION(BlueprintCallable, Category="Game|Cinematic")
+	float GetPlaybackPosition() const;
+
+	/**
+	 * Set the current playback position
+	 * @param NewPlaybackPosition - The new playback position to set.
+	 * If the animation is currently playing, it will continue to do so from the new position
+	 */
+	UFUNCTION(BlueprintCallable, Category="Game|Cinematic")
+	void SetPlaybackPosition(float NewPlaybackPosition);
+
+	/** Check whether the sequence is actively playing. */
+	UFUNCTION(BlueprintCallable, Category="Game|Cinematic")
+	bool IsPlaying() const;
+
+	/** Get the length of the sequence */
+	UFUNCTION(BlueprintCallable, Category="Game|Cinematic")
+	float GetLength() const;
+
+	/** Get the playback rate of this player. */
+	UFUNCTION(BlueprintCallable, Category="Game|Cinematic")
+	float GetPlayRate() const;
+
+	/**
+	 * Set the playback rate of this player. Negative values will play the animation in reverse.
+	 * @param PlayRate - The new rate of playback for the animation.
+	 */
+	UFUNCTION(BlueprintCallable, Category="Game|Cinematic")
+	void SetPlayRate(float PlayRate);
+
+protected:
+
+	// IMovieScenePlayer interface
+
+	virtual void SpawnActorsForMovie(TSharedRef<FMovieSceneSequenceInstance> MovieSceneInstance);
+	virtual void DestroyActorsForMovie(TSharedRef<FMovieSceneSequenceInstance> MovieSceneInstance);
+	virtual void GetRuntimeObjects(TSharedRef<FMovieSceneSequenceInstance> MovieSceneInstance, const FGuid& ObjectHandle, TArray<UObject*>& OutObjects) const override;
+	virtual void UpdateCameraCut(UObject* ObjectToViewThrough, bool bNewCameraCut) const override;
+	virtual void SetViewportSettings(const TMap<FViewportClient*, EMovieSceneViewportParams>& ViewportParamsMap) override;
+	virtual void GetViewportSettings(TMap<FViewportClient*, EMovieSceneViewportParams>& ViewportParamsMap) const override;
+	virtual EMovieScenePlayerStatus::Type GetPlaybackStatus() const override;
+	virtual void AddOrUpdateMovieSceneInstance(UMovieSceneSection& MovieSceneSection, TSharedRef<FMovieSceneSequenceInstance> InstanceToAdd) override;
+	virtual void RemoveMovieSceneInstance(UMovieSceneSection& MovieSceneSection, TSharedRef<FMovieSceneSequenceInstance> InstanceToRemove) override;
+	virtual TSharedRef<FMovieSceneSequenceInstance> GetRootMovieSceneSequenceInstance() const override;
+
+public:
+
+	void Update(const float DeltaSeconds);
+
+private:
+
+	/** Called when the cursor position has changed to implement looping */
+	void OnCursorPositionChanged();
+
+private:
+
+	/** The level sequence to play. */
+	UPROPERTY()
+	ULevelSequenceInstance* LevelSequenceInstance;
+
+	/** Whether we're currently playing. If false, then sequence playback is paused or was never started. */
+	UPROPERTY()
+	bool bIsPlaying;
+
+	/** The current time cursor position within the sequence (in seconds) */
+	UPROPERTY()
+	float TimeCursorPosition;
+
+	/** Specific playback settings for the animation. */
+	UPROPERTY()
+	FLevelSequencePlaybackSettings PlaybackSettings;
+
+	/** The number of times we have looped in the current playback */
+	int32 CurrentNumLoops;
+
+private:
+
+	struct FSpawnedActorInfo
+	{
+		/** Identifier that maps this actor to a movie scene spawnable */
+		FGuid RuntimeGuid;
+
+		/** The spawned actor */
+		TWeakObjectPtr<AActor> SpawnedActor;
+	};
+
+	/** Maps spawnable GUIDs to their spawned actor in the world */
+	TMap<TWeakPtr<FMovieSceneSequenceInstance>, TArray<FSpawnedActorInfo>> InstanceToSpawnedActorMap;
+
+	/** The root movie scene instance to update when playing. */
+	TSharedPtr<FMovieSceneSequenceInstance> RootMovieSceneInstance;
+
+	/** The world this player will spawn actors in, if needed */
+	TWeakObjectPtr<UWorld> World;
+};
diff --git a/Engine/Source/Runtime/LevelSequence/Public/MovieSceneActor.h b/Engine/Source/Runtime/LevelSequence/Public/MovieSceneActor.h
new file mode 100644
--- /dev/null
+++ b/Engine/Source/Runtime/LevelSequence/Public/MovieSceneActor.h
@@ -0,0 +1,60 @@
+// Copyright 1998-2015 Epic Games, Inc. All Rights Reserved.
+
+#pragma once
+
+#include "LevelSequencePlayer.h"
+#include "MovieSceneActor.generated.h"
+
+
+/**
+ * Actor responsible for controlling a specific movie scene sequence in the world
+ */
+UCLASS(Experimental, hideCategories=(Rendering, Physics, LOD, Activation))
+class LEVELSEQUENCE_API AMovieSceneActor
+	: public AActor
+{
+public:
+	AMovieSceneActor(const FObjectInitializer& Init);
+
+	GENERATED_BODY()
+
+	UPROPERTY(EditAnywhere, BlueprintReadOnly, Category="Playback")
+	bool bAutoPlay;
+
+	UPROPERTY(EditAnywhere, BlueprintReadOnly, Category="Playback", meta=(ShowOnlyInnerProperties))
+	FLevelSequencePlaybackSettings PlaybackSettings;
+
+	UPROPERTY(transient, BlueprintReadOnly, Category="Playback")
+	ULevelSequencePlayer* AnimationPlayer;
+
+	UPROPERTY(EditAnywhere, BlueprintReadOnly, Category="General", meta=(AllowedClasses="LevelSequence"))
+	FStringAssetReference LevelSequence;
+
+public:
+
+	UFUNCTION(BlueprintCallable, Category="Game|Cinematic")
+	void SetAnimation(ULevelSequence* Animation);
+
+public:
+
+	virtual void BeginPlay() override;
+	virtual void Tick(float DeltaSeconds) override;
+	virtual void PostInitProperties() override;
+	virtual void PostLoad() override;
+
+#if WITH_EDITOR
+	virtual void PostEditChangeProperty( struct FPropertyChangedEvent& PropertyChangedEvent) override;
+	virtual bool GetReferencedContentObjects(TArray<UObject*>& Objects) const override;
+#endif //WITH_EDITOR
+
+	void InitializePlayer();
+
+private:
+
+	/** Instance of the level sequence we're using. Stored as a property so we can track live actor references. */
+	UPROPERTY(Instanced)
+	ULevelSequenceInstance* SequenceInstance;
+
+	/** Update the instance of our level sequence. */
+	void UpdateAnimationInstance();
+};
diff --git a/Engine/Source/Runtime/MovieSceneCapture/MovieSceneCapture.Build.cs b/Engine/Source/Runtime/MovieSceneCapture/MovieSceneCapture.Build.cs
--- a/Engine/Source/Runtime/MovieSceneCapture/MovieSceneCapture.Build.cs
+++ b/Engine/Source/Runtime/MovieSceneCapture/MovieSceneCapture.Build.cs
@@ -26,7 +26,7 @@ public class MovieSceneCapture : ModuleRules
 
 		PublicDependencyModuleNames.AddRange(
 			new string[] {
-				"ActorAnimation",
+				"LevelSequence",
 			}
 		);
 
diff --git a/Engine/Source/Runtime/MovieSceneCapture/Private/AutomatedActorAnimationCapture.cpp b/Engine/Source/Runtime/MovieSceneCapture/Private/AutomatedActorAnimationCapture.cpp
deleted file mode 100644
diff --git a/Engine/Source/Runtime/MovieSceneCapture/Private/AutomatedLevelSequenceCapture.cpp b/Engine/Source/Runtime/MovieSceneCapture/Private/AutomatedLevelSequenceCapture.cpp
new file mode 100644
--- /dev/null
+++ b/Engine/Source/Runtime/MovieSceneCapture/Private/AutomatedLevelSequenceCapture.cpp
@@ -0,0 +1,68 @@
+// Copyright 1998-2015 Epic Games, Inc. All Rights Reserved.
+
+#include "MovieSceneCapturePCH.h"
+#include "LevelSequencePlayer.h"
+#include "AutomatedLevelSequenceCapture.h"
+#include "ErrorCodes.h"
+
+UAutomatedLevelSequenceCapture::UAutomatedLevelSequenceCapture(const FObjectInitializer& Init)
+	: Super(Init)
+{
+	PlaybackSettings.LoopCount = 0;
+}
+
+void UAutomatedLevelSequenceCapture::Initialize(FViewport* InViewport)
+{
+	if (ULevelSequence* LevelSequenceAsset = Cast<ULevelSequence>(LevelSequence.TryLoad()))
+	{
+		// Set up an instance of the asset for the current world, so we can resolve the bindings
+		AnimationInstance = NewObject<ULevelSequenceInstance>(this, "AnimationInstance");
+		AnimationInstance->Initialize(LevelSequenceAsset, GWorld, false);
+
+		// Set up an animation player that can play back the instance
+		AnimationPlayback = NewObject<ULevelSequencePlayer>(this, "AnimationPlayer");
+		AnimationPlayback->Initialize(AnimationInstance, GWorld, PlaybackSettings);
+
+		// Immediately play/pause the animation. This ensures that the scene is set up, ready for playing.
+		// If a preroll has been specified, this will allow temporal PP effects to settle down before capturing
+		AnimationPlayback->Play();
+		AnimationPlayback->Pause();
+
+		PrerollTime = 0.f;
+	}
+	else
+	{
+		FPlatformMisc::RequestExit(FMovieSceneCaptureExitCodes::AssetNotFound);
+	}
+
+	CaptureStrategy = MakeShareable(new FFixedTimeStepCaptureStrategy(Settings.FrameRate));
+	Super::Initialize(InViewport);
+}
+
+void UAutomatedLevelSequenceCapture::CaptureFrame(float DeltaSeconds)
+{
+	if (AnimationPlayback)
+	{
+		if (PrerollTime.IsSet())
+		{
+			float& Value = PrerollTime.GetValue();
+			Value += DeltaSeconds;
+			if (Value >= PrerollAmount)
+			{
+				AnimationPlayback->Play();
+				StartCapture();
+				PrerollTime.Reset();
+			}
+		}
+		else
+		{
+			Super::CaptureFrame(DeltaSeconds);
+			AnimationPlayback->Update(DeltaSeconds);
+			
+			if (!AnimationPlayback->IsPlaying())
+			{
+				FPlatformMisc::RequestExit(0);
+			}
+		}
+	}
+}
\ No newline at end of file
diff --git a/Engine/Source/Runtime/MovieSceneCapture/Private/MovieSceneCapturePCH.h b/Engine/Source/Runtime/MovieSceneCapture/Private/MovieSceneCapturePCH.h
--- a/Engine/Source/Runtime/MovieSceneCapture/Private/MovieSceneCapturePCH.h
+++ b/Engine/Source/Runtime/MovieSceneCapture/Private/MovieSceneCapturePCH.h
@@ -7,4 +7,4 @@
 
 #include "MovieScene.h"
 #include "ExpressionParser.h"
-#include "ActorAnimationPlayer.h"
\ No newline at end of file
+#include "LevelSequencePlayer.h"
\ No newline at end of file
diff --git a/Engine/Source/Runtime/MovieSceneCapture/Public/AutomatedActorAnimationCapture.h b/Engine/Source/Runtime/MovieSceneCapture/Public/AutomatedActorAnimationCapture.h
deleted file mode 100644
diff --git a/Engine/Source/Runtime/MovieSceneCapture/Public/AutomatedLevelSequenceCapture.h b/Engine/Source/Runtime/MovieSceneCapture/Public/AutomatedLevelSequenceCapture.h
new file mode 100644
--- /dev/null
+++ b/Engine/Source/Runtime/MovieSceneCapture/Public/AutomatedLevelSequenceCapture.h
@@ -0,0 +1,52 @@
+// Copyright 1998-2015 Epic Games, Inc. All Rights Reserved.
+
+#pragma once
+
+#include "MovieSceneCapture.h"
+#include "LevelSequence.h"
+#include "LevelSequencePlayer.h"
+#include "AutomatedLevelSequenceCapture.generated.h"
+
+UCLASS(config=EditorSettings)
+class MOVIESCENECAPTURE_API UAutomatedLevelSequenceCapture : public UMovieSceneCapture
+{
+public:
+	UAutomatedLevelSequenceCapture(const FObjectInitializer&);
+
+	GENERATED_BODY()
+
+	/** The level sequence to use during capture */
+	UPROPERTY(EditAnywhere, Category=General, meta=(AllowedClasses="LevelSequence"))
+	FStringAssetReference LevelSequence;
+
+	/** The level to use for the capture */
+	UPROPERTY(EditAnywhere, Category=General, meta=(AllowedClasses="World"))
+	FStringAssetReference Level;
+
+	/** Specific playback settings */
+	UPROPERTY(config, EditAnywhere, Category="Playback Settings", meta=(ShowOnlyInnerProperties))
+	FLevelSequencePlaybackSettings PlaybackSettings;
+
+	/** The amount of time to wait before playback and capture start. Useful for allowing Post Processing effects to settle down before capturing the animation. */
+	UPROPERTY(config, EditAnywhere, Category=CaptureSettings, AdvancedDisplay, meta=(Units=Seconds, ClampMin=0))
+	float PrerollAmount;
+
+public:
+	// UMovieSceneCapture interface
+	virtual FString GetPackageName() const override { return Level.ToString(); }
+	virtual void Initialize(FViewport* InViewport) override;
+	virtual void CaptureFrame(float DeltaSeconds) override;
+
+private:
+
+	/** Animation instance created at runtime before we start capturing */
+	UPROPERTY(transient)
+	ULevelSequenceInstance* AnimationInstance;
+
+	/** Animation player used to playback the animation at runtime */
+	UPROPERTY(transient)
+	ULevelSequencePlayer* AnimationPlayback;
+
+	/** Transient amount of time that has passed since the level was started. When >= Settings.Preroll, playback will start  */
+	TOptional<float> PrerollTime;
+};
\ No newline at end of file
