Sequencer: Moved level sequencer playback settings customization into LevelSequencerEditor

diff --git a/Engine/Plugins/MovieScene/LevelSequenceEditor/Source/LevelSequenceEditor/LevelSequenceEditor.Build.cs b/Engine/Plugins/MovieScene/LevelSequenceEditor/Source/LevelSequenceEditor/LevelSequenceEditor.Build.cs
--- a/Engine/Plugins/MovieScene/LevelSequenceEditor/Source/LevelSequenceEditor/LevelSequenceEditor.Build.cs
+++ b/Engine/Plugins/MovieScene/LevelSequenceEditor/Source/LevelSequenceEditor/LevelSequenceEditor.Build.cs
@@ -21,9 +21,11 @@ public class LevelSequenceEditor : ModuleRules
 				"CoreUObject",
                 "EditorStyle",
                 "Engine",
+                "InputCore",
 				"LevelEditor",
 				"MovieScene",
 				"MovieSceneTracks",
+                "PropertyEditor",
 				"Sequencer",
                 "Slate",
                 "SlateCore",
@@ -42,6 +44,7 @@ public class LevelSequenceEditor : ModuleRules
             new string[] {
 				"LevelSequenceEditor/Private",
 				"LevelSequenceEditor/Private/AssetTools",
+                "LevelSequenceEditor/Private/Customizations",
 				"LevelSequenceEditor/Private/Factories",
 				"LevelSequenceEditor/Private/Styles",
 			}
diff --git a/Engine/Plugins/MovieScene/LevelSequenceEditor/Source/LevelSequenceEditor/Private/Customizations/LevelSequencePlaybackSettingsCustomization.cpp b/Engine/Plugins/MovieScene/LevelSequenceEditor/Source/LevelSequenceEditor/Private/Customizations/LevelSequencePlaybackSettingsCustomization.cpp
new file mode 100644
--- /dev/null
+++ b/Engine/Plugins/MovieScene/LevelSequenceEditor/Source/LevelSequenceEditor/Private/Customizations/LevelSequencePlaybackSettingsCustomization.cpp
@@ -0,0 +1,141 @@
+﻿// Copyright 1998-2015 Epic Games, Inc. All Rights Reserved.
+
+#include "LevelSequenceEditorPCH.h"
+#include "IDetailPropertyRow.h"
+#include "DetailLayoutBuilder.h"
+#include "DetailWidgetRow.h"
+#include "LevelSequencePlayer.h"
+#include "PropertyHandle.h"
+
+
+#define LOCTEXT_NAMESPACE "LevelSequencePlaybackSettingsCustomization"
+
+
+TSharedRef<IPropertyTypeCustomization> FLevelSequencePlaybackSettingsCustomization::MakeInstance()
+{
+	return MakeShareable(new FLevelSequencePlaybackSettingsCustomization);
+}
+
+
+void FLevelSequencePlaybackSettingsCustomization::CustomizeHeader(TSharedRef<IPropertyHandle> InPropertyHandle, FDetailWidgetRow& HeaderRow, IPropertyTypeCustomizationUtils& CustomizationUtils)
+{
+	// do nothing
+}
+
+
+void FLevelSequencePlaybackSettingsCustomization::CustomizeChildren(TSharedRef<IPropertyHandle> InPropertyHandle, IDetailChildrenBuilder& ChildBuilder, IPropertyTypeCustomizationUtils& CustomizationUtils)
+{
+	StructPropertyHandle = InPropertyHandle;
+	ChildBuilder.AddChildProperty(StructPropertyHandle->GetChildHandle(GET_MEMBER_NAME_CHECKED(FLevelSequencePlaybackSettings, PlayRate)).ToSharedRef());
+	LoopCountProperty = StructPropertyHandle->GetChildHandle(GET_MEMBER_NAME_CHECKED(FLevelSequencePlaybackSettings, LoopCount));
+
+	// Set up the initial environment
+	{
+		LoopModes.Add(MakeShareable( new FLoopMode{ LOCTEXT("DontLoop", "Don't Loop"), 0 } ));
+		LoopModes.Add(MakeShareable( new FLoopMode{ LOCTEXT("Indefinitely", "Loop Indefinitely"), -1 }) );
+		LoopModes.Add(MakeShareable( new FLoopMode{ LOCTEXT("Exactly", "Loop Exactly..."), 1 }) );
+
+		int32 CurrentValue = -1;
+		LoopCountProperty->GetValue(CurrentValue);
+
+		if (auto* FoundMode = LoopModes.FindByPredicate([&](const TSharedPtr<FLoopMode>& In){ return In->Value == CurrentValue; }) )
+		{
+			CurrentMode = *FoundMode;
+		}
+		else
+		{
+			CurrentMode = LoopModes.Last();
+		}
+	}
+
+	ChildBuilder.AddChildContent(LOCTEXT("LoopTitle", "Loop"))
+		.NameContent()
+		[
+			LoopCountProperty->CreatePropertyNameWidget()
+		]
+		.ValueContent()
+		.HAlign(HAlign_Fill)
+		.MaxDesiredWidth(200)
+		[
+			SNew(SHorizontalBox)
+
+			+ SHorizontalBox::Slot()
+				.VAlign(VAlign_Center)
+				.AutoWidth()
+				[
+					SNew(SComboBox<TSharedPtr<FLoopMode>>)
+						.OptionsSource(&LoopModes)
+						.OnSelectionChanged_Lambda([&](TSharedPtr<FLoopMode> Mode, ESelectInfo::Type){
+							CurrentMode = Mode;
+							UpdateProperty();
+						})
+						.OnGenerateWidget_Lambda([&](TSharedPtr<FLoopMode> InMode){
+							return SNew(STextBlock)
+								.Font(IDetailLayoutBuilder::GetDetailFont())
+								.Text(InMode->DisplayName);
+						})
+						.InitiallySelectedItem(CurrentMode)
+						[
+							SAssignNew(CurrentText, STextBlock)
+								.Font(IDetailLayoutBuilder::GetDetailFont())
+								.Text(CurrentMode->DisplayName)
+						]
+				]
+
+			+ SHorizontalBox::Slot()
+				.Padding(FMargin(4.f,0))
+				.VAlign(VAlign_Center)
+				[
+					SAssignNew(LoopEntry, SHorizontalBox)
+						.Visibility(CurrentMode == LoopModes.Last() ? EVisibility::Visible : EVisibility::Collapsed)
+
+					+ SHorizontalBox::Slot()
+						.VAlign(VAlign_Center)
+						.Padding(FMargin(0, 0, 4.f, 0))
+						[
+							LoopCountProperty->CreatePropertyValueWidget()
+						]
+
+					+ SHorizontalBox::Slot()
+						.VAlign(VAlign_Center)
+						.AutoWidth()
+						[
+							SNew(STextBlock)
+								.Font(IDetailLayoutBuilder::GetDetailFont())
+								.Text_Raw(this, &FLevelSequencePlaybackSettingsCustomization::GetCustomSuffix)
+						]
+				]
+		];
+}
+
+
+FText FLevelSequencePlaybackSettingsCustomization::GetCustomSuffix() const
+{
+	int32 NumLoops = 0;
+	LoopCountProperty->GetValue(NumLoops);
+
+	return NumLoops == 1
+		? LOCTEXT("Time", "time")
+		: LOCTEXT("Times", "times");
+}
+
+
+void FLevelSequencePlaybackSettingsCustomization::UpdateProperty()
+{
+	if (CurrentMode == LoopModes.Last())
+	{
+		// Show custom loop entry
+		LoopEntry->SetVisibility(EVisibility::Visible);
+	}
+	else
+	{
+		// Hide custom loop entry
+		LoopEntry->SetVisibility(EVisibility::Collapsed);
+	}
+
+	LoopCountProperty->SetValue(CurrentMode->Value);
+	CurrentText->SetText(CurrentMode->DisplayName);
+}
+
+
+#undef LOCTEXT_NAMESPACE
diff --git a/Engine/Plugins/MovieScene/LevelSequenceEditor/Source/LevelSequenceEditor/Private/Customizations/LevelSequencePlaybackSettingsCustomization.h b/Engine/Plugins/MovieScene/LevelSequenceEditor/Source/LevelSequenceEditor/Private/Customizations/LevelSequencePlaybackSettingsCustomization.h
new file mode 100644
--- /dev/null
+++ b/Engine/Plugins/MovieScene/LevelSequenceEditor/Source/LevelSequenceEditor/Private/Customizations/LevelSequencePlaybackSettingsCustomization.h
@@ -0,0 +1,53 @@
+// Copyright 1998-2015 Epic Games, Inc. All Rights Reserved.
+
+#pragma once
+
+
+/**
+ * Details view customization for the FLevelSequencePlaybackSettings struct.
+ */
+class FLevelSequencePlaybackSettingsCustomization
+	: public IPropertyTypeCustomization
+{
+public:
+
+	/** Makes a new instance of this detail layout class for a specific detail view requesting it. */
+	static TSharedRef<IPropertyTypeCustomization> MakeInstance();
+
+public:
+
+	// IDetailCustomization interface
+
+	virtual void CustomizeHeader( TSharedRef<IPropertyHandle> PropertyHandle, FDetailWidgetRow& HeaderRow, IPropertyTypeCustomizationUtils& CustomizationUtils ) override;
+	virtual void CustomizeChildren( TSharedRef<IPropertyHandle> PropertyHandle, IDetailChildrenBuilder& ChildBuilder, IPropertyTypeCustomizationUtils& CustomizationUtils ) override;
+
+protected:
+
+	/** Respond to a selection change event from the combo box. */
+	void UpdateProperty();
+	
+private:
+
+	/** Get the suffix to display after the custom entry box. */
+	FText GetCustomSuffix() const;
+
+	/** Array of loop modes. */
+	struct FLoopMode
+	{
+		FText DisplayName;
+		int32 Value;
+	};
+
+	TArray<TSharedPtr<FLoopMode>> LoopModes;
+
+	/** The loop mode we're currently displaying. */
+	TSharedPtr<FLoopMode> CurrentMode;
+	/** The text of the current selection. */
+	TSharedPtr<STextBlock> CurrentText;
+	/** The loop number entry to be hidden and shown based on combo box selection. */
+	TSharedPtr<SWidget> LoopEntry;
+
+	/** Property handles of the properties we're editing */
+	TSharedPtr<IPropertyHandle> StructPropertyHandle;
+	TSharedPtr<IPropertyHandle> LoopCountProperty;
+};
diff --git a/Engine/Plugins/MovieScene/LevelSequenceEditor/Source/LevelSequenceEditor/Private/LevelSequenceEditorModule.cpp b/Engine/Plugins/MovieScene/LevelSequenceEditor/Source/LevelSequenceEditor/Private/LevelSequenceEditorModule.cpp
--- a/Engine/Plugins/MovieScene/LevelSequenceEditor/Source/LevelSequenceEditor/Private/LevelSequenceEditorModule.cpp
+++ b/Engine/Plugins/MovieScene/LevelSequenceEditor/Source/LevelSequenceEditor/Private/LevelSequenceEditorModule.cpp
@@ -1,10 +1,11 @@
 // Copyright 1998-2015 Epic Games, Inc. All Rights Reserved.
 
 #include "LevelSequenceEditorPCH.h"
+#include "LevelEditor.h"
 #include "LevelSequenceEditorStyle.h"
 #include "ModuleInterface.h"
-#include "LevelEditor.h"
 #include "MovieSceneActor.h"
+#include "PropertyEditorModule.h"
 
 #define LOCTEXT_NAMESPACE "LevelSequenceEditor"
 
@@ -44,12 +45,14 @@ public:
 		Style = MakeShareable(new FLevelSequenceEditorStyle());
 
 		RegisterAssetTools();
+		RegisterCustomizations();
 		RegisterMenuExtensions();
 	}
 	
 	virtual void ShutdownModule() override
 	{
 		UnregisterAssetTools();
+		UnregisterCustomizations();
 		UnregisterMenuExtensions();
 	}
 
@@ -91,6 +94,22 @@ protected:
 		}
 	}
 
+protected:
+
+	/** Register details view customizations. */
+	void RegisterCustomizations()
+	{
+		FPropertyEditorModule& PropertyModule = FModuleManager::LoadModuleChecked<FPropertyEditorModule>("PropertyEditor");
+		PropertyModule.RegisterCustomPropertyTypeLayout(FLevelSequencePlaybackSettings::StaticStruct()->GetFName(), FOnGetPropertyTypeCustomizationInstance::CreateStatic(&FLevelSequencePlaybackSettingsCustomization::MakeInstance));
+	}
+
+	/** Unregister details view customizations. */
+	void UnregisterCustomizations()
+	{
+		FPropertyEditorModule& PropertyModule = FModuleManager::LoadModuleChecked<FPropertyEditorModule>("PropertyEditor");
+		PropertyModule.UnregisterCustomPropertyTypeLayout(FLevelSequencePlaybackSettings::StaticStruct()->GetFName());
+	}
+
 protected:
 
 	/** Register menu extensions for the level editor toolbar. */
diff --git a/Engine/Plugins/MovieScene/LevelSequenceEditor/Source/LevelSequenceEditor/Private/LevelSequenceEditorPCH.h b/Engine/Plugins/MovieScene/LevelSequenceEditor/Source/LevelSequenceEditor/Private/LevelSequenceEditorPCH.h
--- a/Engine/Plugins/MovieScene/LevelSequenceEditor/Source/LevelSequenceEditor/Private/LevelSequenceEditorPCH.h
+++ b/Engine/Plugins/MovieScene/LevelSequenceEditor/Source/LevelSequenceEditor/Private/LevelSequenceEditorPCH.h
@@ -21,3 +21,4 @@
 
 #include "LevelSequenceActions.h"
 #include "LevelSequenceEditorToolkit.h"
+#include "LevelSequencePlaybackSettingsCustomization.h"
diff --git a/Engine/Source/Editor/DetailCustomizations/Private/DetailCustomizations.cpp b/Engine/Source/Editor/DetailCustomizations/Private/DetailCustomizations.cpp
--- a/Engine/Source/Editor/DetailCustomizations/Private/DetailCustomizations.cpp
+++ b/Engine/Source/Editor/DetailCustomizations/Private/DetailCustomizations.cpp
@@ -104,10 +104,10 @@
 #include "ConfigEditorPropertyDetails.h"
 #include "AssetImportDataCustomization.h"
 #include "CaptureResolutionCustomization.h"
-#include "LevelSequencePlaybackSettingsCustomization.h"
 #include "TextCustomization.h"
 #include "AnimTrailNodeDetails.h"
 
+
 IMPLEMENT_MODULE( FDetailCustomizationsModule, DetailCustomizations );
 
 
@@ -116,7 +116,6 @@ void FDetailCustomizationsModule::StartupModule()
 	FPropertyEditorModule& PropertyModule = FModuleManager::GetModuleChecked<FPropertyEditorModule>("PropertyEditor");
 	
 	RegisterPropertyTypeCustomizations();
-
 	RegisterObjectCustomizations();
 
 	PropertyModule.NotifyCustomizationModuleChanged();
@@ -213,7 +212,6 @@ void FDetailCustomizationsModule::RegisterPropertyTypeCustomizations()
 	RegisterCustomPropertyTypeLayout("HierarchicalSimplification", FOnGetPropertyTypeCustomizationInstance::CreateStatic(&FHierarchicalSimplificationCustomizations::MakeInstance));
 	RegisterCustomPropertyTypeLayout("PostProcessSettings", FOnGetPropertyTypeCustomizationInstance::CreateStatic(&FPostProcessSettingsCustomization::MakeInstance));
 	RegisterCustomPropertyTypeLayout("AssetImportInfo", FOnGetPropertyTypeCustomizationInstance::CreateStatic(&FAssetImportDataCustomization::MakeInstance));
-	RegisterCustomPropertyTypeLayout("LevelSequencePlaybackSettings", FOnGetPropertyTypeCustomizationInstance::CreateStatic(&FLevelSequencePlaybackSettingsCustomization::MakeInstance));
 	RegisterCustomPropertyTypeLayout("CaptureResolution", FOnGetPropertyTypeCustomizationInstance::CreateStatic(&FCaptureResolutionCustomization::MakeInstance));
 	RegisterCustomPropertyTypeLayout("WeightedBlendable", FOnGetPropertyTypeCustomizationInstance::CreateStatic(&FWeightedBlendableCustomization::MakeInstance));
 }
@@ -330,4 +328,4 @@ void FDetailCustomizationsModule::RegisterCustomPropertyTypeLayout(FName Propert
 	static FName PropertyEditor("PropertyEditor");
 	FPropertyEditorModule& PropertyModule = FModuleManager::GetModuleChecked<FPropertyEditorModule>(PropertyEditor);
 	PropertyModule.RegisterCustomPropertyTypeLayout(PropertyTypeName, PropertyTypeLayoutDelegate);
-}
\ No newline at end of file
+}
diff --git a/Engine/Source/Editor/DetailCustomizations/Private/LevelSequencePlaybackSettingsCustomization.cpp b/Engine/Source/Editor/DetailCustomizations/Private/LevelSequencePlaybackSettingsCustomization.cpp
deleted file mode 100644
diff --git a/Engine/Source/Editor/DetailCustomizations/Private/LevelSequencePlaybackSettingsCustomization.h b/Engine/Source/Editor/DetailCustomizations/Private/LevelSequencePlaybackSettingsCustomization.h
deleted file mode 100644
