Temp fix for unregistering details customizations on shutdown currently not working

diff --git a/Engine/Plugins/Media/MediaPlayerEditor/Source/MediaPlayerEditor/Private/MediaPlayerEditorModule.cpp b/Engine/Plugins/Media/MediaPlayerEditor/Source/MediaPlayerEditor/Private/MediaPlayerEditorModule.cpp
--- a/Engine/Plugins/Media/MediaPlayerEditor/Source/MediaPlayerEditor/Private/MediaPlayerEditorModule.cpp
+++ b/Engine/Plugins/Media/MediaPlayerEditor/Source/MediaPlayerEditor/Private/MediaPlayerEditorModule.cpp
@@ -53,7 +53,7 @@ public:
 	virtual void ShutdownModule() override
 	{
 		UnregisterAssetTools();
-		UnregisterCustomizations();
+		//UnregisterCustomizations();
 		UnregisterMenuExtensions();
 		UnregisterThumbnailRenderers();
 	}
diff --git a/Engine/Plugins/MovieScene/LevelSequenceEditor/Source/LevelSequenceEditor/Private/LevelSequenceEditorModule.cpp b/Engine/Plugins/MovieScene/LevelSequenceEditor/Source/LevelSequenceEditor/Private/LevelSequenceEditorModule.cpp
--- a/Engine/Plugins/MovieScene/LevelSequenceEditor/Source/LevelSequenceEditor/Private/LevelSequenceEditorModule.cpp
+++ b/Engine/Plugins/MovieScene/LevelSequenceEditor/Source/LevelSequenceEditor/Private/LevelSequenceEditorModule.cpp
@@ -52,7 +52,7 @@ public:
 	virtual void ShutdownModule() override
 	{
 		UnregisterAssetTools();
-		UnregisterCustomizations();
+		//UnregisterCustomizations();
 		UnregisterMenuExtensions();
 	}
 
