Rollback incorrectly deleted files from Orion merge

diff --git a/Engine/Plugins/MovieScene/LevelSequenceEditor/Source/LevelSequenceEditor/Private/AssetTools/LevelSequenceActions.h b/Engine/Plugins/MovieScene/LevelSequenceEditor/Source/LevelSequenceEditor/Private/AssetTools/LevelSequenceActions.h
new file mode 100644
--- /dev/null
+++ b/Engine/Plugins/MovieScene/LevelSequenceEditor/Source/LevelSequenceEditor/Private/AssetTools/LevelSequenceActions.h
@@ -0,0 +1,38 @@
+// Copyright 1998-2015 Epic Games, Inc. All Rights Reserved.
+
+#pragma once
+
+#include "AssetTypeActions_Base.h"
+
+
+/**
+ * Implements actions for ULevelSequence assets.
+ */
+class FLevelSequenceActions
+	: public FAssetTypeActions_Base
+{
+public:
+
+	/**
+	 * Creates and initializes a new instance.
+	 *
+	 * @param InStyle The style set to use for asset editor toolkits.
+	 */
+	FLevelSequenceActions(const TSharedRef<ISlateStyle>& InStyle);
+
+public:
+	
+	// IAssetTypeActions interface
+
+	virtual uint32 GetCategories() override;
+	virtual FText GetName() const override;
+	virtual UClass* GetSupportedClass() const override;
+	virtual FColor GetTypeColor() const override;
+	virtual void OpenAssetEditor(const TArray<UObject*>& InObjects, TSharedPtr<class IToolkitHost> EditWithinLevelEditor = TSharedPtr<IToolkitHost>()) override;
+	virtual bool ShouldForceWorldCentric() override;
+
+private:
+
+	/** Pointer to the style set to use for toolkits. */
+	TSharedRef<ISlateStyle> Style;
+};
diff --git a/Engine/Plugins/MovieScene/LevelSequenceEditor/Source/LevelSequenceEditor/Private/Factories/LevelSequenceFactoryNew.h b/Engine/Plugins/MovieScene/LevelSequenceEditor/Source/LevelSequenceEditor/Private/Factories/LevelSequenceFactoryNew.h
new file mode 100644
--- /dev/null
+++ b/Engine/Plugins/MovieScene/LevelSequenceEditor/Source/LevelSequenceEditor/Private/Factories/LevelSequenceFactoryNew.h
@@ -0,0 +1,24 @@
+// Copyright 1998-2015 Epic Games, Inc. All Rights Reserved.
+
+#pragma once
+
+#include "Factories/Factory.h"
+#include "LevelSequenceFactoryNew.generated.h"
+
+
+/**
+ * Implements a factory for ULevelSequence objects.
+ */
+UCLASS(hidecategories=Object)
+class ULevelSequenceFactoryNew
+	: public UFactory
+{
+	GENERATED_UCLASS_BODY()
+
+public:
+
+	// UFactory Interface
+
+	virtual UObject* FactoryCreateNew(UClass* Class, UObject* InParent, FName Name, EObjectFlags Flags, UObject* Context, FFeedbackContext* Warn) override;
+	virtual bool ShouldShowInNewMenu() const override;
+};
