Copying //UE4/Orion-Staging to Dev-Main (Originating from //Orion/Main at CL-2777663)

 Change 2777555 on 2015/11/23 by [REDACTED_PERSON]

	Friend List Sub Menu restyling to new designs

	#TESTS Check 3 sub menus of friends list (Online Status, Friends List, Settings) still function correctly with new styling.

Change 2777506 on 2015/11/23 by [REDACTED_PERSON]

	Sequencer: Copy/Paste command binding is no longer active if the seuqnece widget is not focused

	This addresses UE-23423

	#tests Tested copy/paste inside and outside of sequencer

Change 2777505 on 2015/11/23 by [REDACTED_PERSON]

	Sequencer: Undoing the addition of a spawnable now ensures its actor instance is deleted correctly

	This addresses UE-23450

	#tests tested the repro steps on the bug

Change 2777489 on 2015/11/23 by [REDACTED_PERSON]

	Sequencer: Workflow optimizations for spawnables

	 - Editing a property on an instance of a spawnable now automatically propagates to the spawnable defaults, provided the property is not keyed
	 - Fixed a few cases where spawnables were left lingering around while scrubbing or switching between sub-sequences
	 - Fixed the root sequence instance being evaluated when there was a sub-sequence focused.
	 - Selection states are now remembered for spawnable objects when they are destroyed/re-spawned

	#tests tested loks of object types as spawnables in PIE and in editor

Change 2777321 on 2015/11/23 by [REDACTED_PERSON]

	Updated the usage of World->UpdateStreamingLevels to FlushStreaming levels on the advice of Dmitriy. Need this to ensure that the streaming is completed before we send it off to lightmass.

	#Tests Run the RebuildLightmaps UAT script.

Change 2777091 on 2015/11/22 by [REDACTED_PERSON]

	Changed "inappropriate outmost" warning on package load to an error. At the very least we want this for a day or two on Orion to surface errors quickly, but may be a good thing to make a standard error since it indicates something that's likely broken.

	#tests Golden path in game, cooked content

Change 2777037 on 2015/11/22 by [REDACTED_PERSON]

	Additional debug info to track down https://jira.ol.epicgames.net/browse/OR-9675

	#tests Golden path (PIE) + compiled for PS4

Change 2777030 on 2015/11/22 by [REDACTED_PERSON]

	#online,externalui,ps4
	- added footer option for closing the embedded web browser

	#tests exec command to try on ps4

Change 2777019 on 2015/11/22 by [REDACTED_PERSON]

	Possible fix for OR-9851
	#test GoldenPath, PS4

Change 2776932 on 2015/11/22 by [REDACTED_PERSON]

	Sequencer: Fix editor selection so that it's not modified in response to the sequencer outliner tree node changing selection if the user is not explicitly selecting in the tree.

	#tests Select an actor that Sequencer doesn't control and it shouldn't deselect.

Change 2776900 on 2015/11/21 by [REDACTED_PERSON]

	HighQuality particle lights.
	#test GoldenPath, HQ Particles w/wo ShadowCasting.

Change 2776868 on 2015/11/21 by [REDACTED_PERSON]

	Reduced temporal aa responsiveness back where it was.

Change 2776867 on 2015/11/21 by [REDACTED_PERSON]

	Removed shading terminator bias meant for shadow map acne but it made character faces look worse.

Change 2776840 on 2015/11/21 by [REDACTED_PERSON]

	Hair indirect lighting implemented.

	Improvements to hair shading model. No longer uses backlit parameter.

	#tests editor

Change 2776748 on 2015/11/21 by [REDACTED_PERSON]

	Sequencer: Continued to implement track label editor

	Note: still disabled, because there are a couple remaining issues

	#test: Editor, Runtime

Change 2776493 on 2015/11/20 by [REDACTED_PERSON]

	Sequencer: Wrapped the node tree context menu actions for editing in an 'Edit' section

diff --git a/Engine/Config/BaseEngine.ini b/Engine/Config/BaseEngine.ini
--- a/Engine/Config/BaseEngine.ini
+++ b/Engine/Config/BaseEngine.ini
@@ -756,6 +756,9 @@ bUseStreamingPause=false
 +TaggedPropertyRedirects=(ClassName="PlayerController",OldPropertyName="PlayerCameraClass",NewPropertyName="PlayerCameraManagerClass")
 +K2FieldRedirects=(OldFieldName="GameplayStatics.GetPlayerCamera",NewFieldName="GameplayStatics.GetPlayerCameraManager")
 
+; Renamed FMovieScenePossessable::ParentSpawnableGuid to ParentGuid
++TaggedPropertyRedirects=(ClassName="FMovieScenePossessable",OldPropertyName="ParentSpawnableGuid",NewPropertyName="ParentGuid")
+
 ; Enum redirects
 +EnumRedirects=(EnumName="ECollisionChannel",OldEnumEntry="ECC_Default",NewEnumEntry="ECC_Visibility")
 +EnumRedirects=(EnumName="ECollisionChannel",OldEnumEntry="ECC_WorldTrace",NewEnumEntry="ECC_WorldStatic")
diff --git a/Engine/Plugins/MovieScene/LevelSequenceEditor/Source/LevelSequenceEditor/Private/LevelSequenceEditorToolkit.cpp b/Engine/Plugins/MovieScene/LevelSequenceEditor/Source/LevelSequenceEditor/Private/LevelSequenceEditorToolkit.cpp
--- a/Engine/Plugins/MovieScene/LevelSequenceEditor/Source/LevelSequenceEditor/Private/LevelSequenceEditorToolkit.cpp
+++ b/Engine/Plugins/MovieScene/LevelSequenceEditor/Source/LevelSequenceEditor/Private/LevelSequenceEditorToolkit.cpp
@@ -11,9 +11,12 @@
 #include "SDockTab.h"
 #include "MovieSceneBinding.h"
 #include "MovieScene.h"
+#include "MovieSceneSequenceInstance.h"
 #include "MovieSceneMaterialTrack.h"
 #include "ScopedTransaction.h"
 #include "ClassIconFinder.h"
+#include "LevelSequenceSpawnRegister.h"
+#include "ISequencerObjectChangeListener.h"
 
 #include "SceneOutlinerModule.h"
 #include "SceneOutlinerPublicTypes.h"
@@ -21,6 +24,177 @@
 
 #define LOCTEXT_NAMESPACE "LevelSequenceEditor"
 
+/**
+ * Spawn register used in the editor to add some usability features like maintaining selection states, and projecting spawned state onto spawnable defaults
+ */
+class FLevelSequenceEditorSpawnRegister : public FLevelSequenceSpawnRegister
+{
+public:
+	/** Weak pointer to the active sequencer */
+	TWeakPtr<ISequencer> WeakSequencer;
+
+	/** Constructor */
+	FLevelSequenceEditorSpawnRegister()
+	{
+		bShouldClearSelectionCache = true;
+
+		FLevelEditorModule& LevelEditor = FModuleManager::GetModuleChecked<FLevelEditorModule>("LevelEditor");
+		OnActorSelectionChangedHandle = LevelEditor.OnActorSelectionChanged().AddRaw(this, &FLevelSequenceEditorSpawnRegister::OnActorSelectionChanged);
+
+		OnActorMovedHandle = GEditor->OnActorMoved().AddLambda([=](AActor* Actor){
+			OnSpawnedObjectPropertyChanged(*Actor);
+		});
+	}
+
+	~FLevelSequenceEditorSpawnRegister()
+	{
+		GEditor->OnActorMoved().Remove(OnActorMovedHandle);
+		if (FLevelEditorModule* LevelEditor = FModuleManager::GetModulePtr<FLevelEditorModule>("LevelEditor"))
+		{
+			LevelEditor->OnActorSelectionChanged().Remove(OnActorSelectionChangedHandle);
+		}
+	}
+
+private:
+
+	/** Called to spawn an object */
+	virtual UObject* SpawnObject(const FGuid& BindingId, FMovieSceneSequenceInstance& SequenceInstance, IMovieScenePlayer& Player) override
+	{
+		TGuardValue<bool> Guard(bShouldClearSelectionCache, false);
+
+		UObject* NewObject = FLevelSequenceSpawnRegister::SpawnObject(BindingId, SequenceInstance, Player);
+
+		// Add an object listener for the spawned object to propagate changes back onto the spawnable default
+		TSharedPtr<ISequencer> Sequencer = WeakSequencer.Pin();
+		if (Sequencer.IsValid() && NewObject)
+		{
+			Sequencer->GetObjectChangeListener().GetOnAnyPropertyChanged(*NewObject).AddSP(this, &FLevelSequenceEditorSpawnRegister::OnSpawnedObjectPropertyChanged);
+
+			// Select the actor if we think it should be selected
+			AActor* Actor = Cast<AActor>(NewObject);
+			if (Actor && SelectedSpawnedObjects.Contains(FMovieSceneSpawnRegisterKey(BindingId, SequenceInstance)))
+			{
+				GEditor->SelectActor(Actor, true /*bSelected*/, true /*bNotify*/);
+			}
+		}
+
+		return NewObject;
+	}
+
+	/** Called right before an object is about to be destroyed */
+	virtual void PreDestroyObject(UObject& Object, const FGuid& BindingId, FMovieSceneSequenceInstance& SequenceInstance) override
+	{
+		TGuardValue<bool> Guard(bShouldClearSelectionCache, false);
+
+		// Cache its selection state
+		AActor* Actor = Cast<AActor>(&Object);
+		if (Actor && GEditor->GetSelectedActors()->IsSelected(Actor))
+		{
+			SelectedSpawnedObjects.Add(FMovieSceneSpawnRegisterKey(BindingId, SequenceInstance));
+			GEditor->SelectActor(Actor, false /*bSelected*/, true /*bNotify*/);
+		}
+
+		// Remove our object listener
+		TSharedPtr<ISequencer> Sequencer = WeakSequencer.Pin();
+		if (Sequencer.IsValid())
+		{
+			Sequencer->GetObjectChangeListener().ReportObjectDestroyed(Object);
+		}
+
+		FLevelSequenceSpawnRegister::PreDestroyObject(Object, BindingId, SequenceInstance);
+	}
+
+	/** Populate a map of properties that are keyed on a particular object */
+	void PopulateKeyedPropertyMap(AActor& SpawnedObject, TMap<UObject*, TSet<UProperty*>>& OutKeyedPropertyMap)
+	{
+		TSharedPtr<ISequencer> Sequencer = WeakSequencer.Pin();
+
+		Sequencer->GetAllKeyedProperties(SpawnedObject, OutKeyedPropertyMap.FindOrAdd(&SpawnedObject));
+
+		for (UActorComponent* Component : SpawnedObject.GetComponents())
+		{
+			Sequencer->GetAllKeyedProperties(*Component, OutKeyedPropertyMap.FindOrAdd(Component));
+		}
+	}
+
+	/** Called when the editor selection has changed */
+	void OnActorSelectionChanged(const TArray<UObject*>& NewSelection, bool bForceRefresh)
+	{
+		if (bShouldClearSelectionCache)
+		{
+			SelectedSpawnedObjects.Reset();
+		}
+	}
+
+	/** Called when a property on a spawned object changes */
+	void OnSpawnedObjectPropertyChanged(UObject& SpawnedObject)
+	{
+		using namespace EditorUtilities;
+
+		AActor* Actor = CastChecked<AActor>(&SpawnedObject);
+		if (!Actor)
+		{
+			return;
+		}
+
+		TMap<UObject*, TSet<UProperty*>> ObjectToKeyedProperties;
+		PopulateKeyedPropertyMap(*Actor, ObjectToKeyedProperties);
+
+		// Copy any changed actor properties onto the default actor, provided they are not keyed
+		FCopyOptions Options(ECopyOptions::PropagateChangesToArchetypeInstances);
+
+		// Set up a property filter so only stuff that is not keyed gets copied onto the default
+		Options.PropertyFilter = [&](const UProperty& Property, const UObject& Object) -> bool {
+			const TSet<UProperty*>* ExcludedProperties = ObjectToKeyedProperties.Find(const_cast<UObject*>(&Object));
+
+			return !ExcludedProperties || !ExcludedProperties->Contains(const_cast<UProperty*>(&Property));
+		};
+
+		// Now copy the actor properties
+		AActor* DefaultActor = Actor->GetClass()->GetDefaultObject<AActor>();
+		EditorUtilities::CopyActorProperties(Actor, DefaultActor, Options);
+
+		// The above function call explicitly doesn't copy the root component transform (so the default actor is always at 0,0,0)
+		// But in sequencer, we want the object to have a default transform if it doesn't have a transform track
+		static FName RelativeLocation = GET_MEMBER_NAME_CHECKED(USceneComponent, RelativeLocation);
+		static FName RelativeRotation = GET_MEMBER_NAME_CHECKED(USceneComponent, RelativeRotation);
+		static FName RelativeScale3D = GET_MEMBER_NAME_CHECKED(USceneComponent, RelativeScale3D);
+
+		bool bHasKeyedTransform = false;
+		for (UProperty* Property : *ObjectToKeyedProperties.Find(Actor))
+		{
+			FName PropertyName = Property->GetFName();
+			bHasKeyedTransform = PropertyName == RelativeLocation || PropertyName == RelativeRotation || PropertyName == RelativeScale3D;
+			if (bHasKeyedTransform)
+			{
+				break;
+			}
+		}
+
+		// Set the default transform if it's not keyed
+		USceneComponent* RootComponent = Actor->GetRootComponent();
+		USceneComponent* DefaultRootComponent = DefaultActor->GetRootComponent();
+
+		if (!bHasKeyedTransform && RootComponent && DefaultRootComponent)
+		{
+			DefaultRootComponent->RelativeLocation = RootComponent->RelativeLocation;
+			DefaultRootComponent->RelativeRotation = RootComponent->RelativeRotation;
+			DefaultRootComponent->RelativeScale3D = RootComponent->RelativeScale3D;
+		}
+	}
+
+private:
+
+	/** Handles for delegates that we've bound to */
+	FDelegateHandle OnActorMovedHandle, OnActorSelectionChangedHandle;
+
+	/** Set of spawn register keys for objects that should be selected if they are spawned */
+	TSet<FMovieSceneSpawnRegisterKey> SelectedSpawnedObjects;
+
+	/** True if we should clear the above selection cache when the editor selection has been changed */
+	bool bShouldClearSelectionCache;
+};
+
 
 /* Local constants
  *****************************************************************************/
@@ -83,26 +257,20 @@ void FLevelSequenceEditorToolkit::Initialize( const EToolkitMode::Type Mode, con
 
 	LevelSequence = InLevelSequence;
 
-	// Find the current editor world
-	for (auto& Context : GEngine->GetWorldContexts())
-	{
-		if (Context.WorldType == EWorldType::Editor)
-		{
-			LevelSequence->BindToContext(Context.World());
-		}
-	}
-
 	const bool bCreateDefaultStandaloneMenu = true;
 	const bool bCreateDefaultToolbar = false;
 
 	FAssetEditorToolkit::InitAssetEditor(Mode, InitToolkitHost, SequencerDefs::SequencerAppIdentifier, StandaloneDefaultLayout, bCreateDefaultStandaloneMenu, bCreateDefaultToolbar, LevelSequence);
 
+	TSharedRef<FLevelSequenceEditorSpawnRegister> SpawnRegister = MakeShareable(new FLevelSequenceEditorSpawnRegister);
+
 	// initialize sequencer
 	FSequencerInitParams SequencerInitParams;
 	{
 		SequencerInitParams.RootSequence = LevelSequence;
 		SequencerInitParams.bEditWithinLevelEditor = bEditWithinLevelEditor;
 		SequencerInitParams.ToolkitHost = InitToolkitHost;
+		SequencerInitParams.SpawnRegister = SpawnRegister;
 
 		TSharedRef<FExtender> AddMenuExtender = MakeShareable(new FExtender);
 
@@ -125,6 +293,8 @@ void FLevelSequenceEditorToolkit::Initialize( const EToolkitMode::Type Mode, con
 	}
 
 	Sequencer = FModuleManager::LoadModuleChecked<ISequencerModule>("Sequencer").CreateSequencer(SequencerInitParams);
+	
+	SpawnRegister->WeakSequencer = Sequencer;
 
 	if (bEditWithinLevelEditor)
 	{
@@ -257,14 +427,9 @@ void FLevelSequenceEditorToolkit::AddActorsToSequencer(AActor*const* InActors, i
 	while (NumActors--)
 	{
 		AActor* ThisActor = *InActors;
-
-		FGuid ObjectGuid = FocussedSequence->FindObjectId(*ThisActor);
-
-		// Add this object if it hasn't already been possessed
-		if (!ObjectGuid.IsValid() || !FocussedMovieScene->FindPossessable(ObjectGuid))
+		if (!Sequencer->GetFocusedMovieSceneSequenceInstance()->FindObjectId(*ThisActor).IsValid())
 		{
-			ObjectGuid = FocussedMovieScene->AddPossessable(ThisActor->GetActorLabel(), ThisActor->GetClass());
-			FocussedSequence->BindPossessableObject(ObjectGuid, *ThisActor);
+			Sequencer->CreateBinding(*ThisActor, ThisActor->GetActorLabel());
 		}
 
 		GEditor->SelectActor(ThisActor, true, true);
@@ -357,8 +522,6 @@ void FLevelSequenceEditorToolkit::AddPosessActorMenuExtensions(FMenuBuilder& Men
 
 void FLevelSequenceEditorToolkit::HandleMapChanged(class UWorld* NewWorld, EMapChangeType MapChangeType)
 {
-	LevelSequence->BindToContext(NewWorld);
-
 	Sequencer->NotifyMapChanged(NewWorld, MapChangeType);
 }
 
diff --git a/Engine/Programs/CrashReportClient/Config/DefaultEngine.ini b/Engine/Programs/CrashReportClient/Config/DefaultEngine.ini
--- a/Engine/Programs/CrashReportClient/Config/DefaultEngine.ini
+++ b/Engine/Programs/CrashReportClient/Config/DefaultEngine.ini
@@ -1,4 +1,5 @@
 [CrashReportClient]
 CrashReportReceiverIP="http://crashreporter.epicgames.com:57005"
 bAllowToBeContacted=true
-bSendLogFile=true
\ No newline at end of file
+bSendLogFile=true
+bHideLogFilesOption=true
\ No newline at end of file
diff --git a/Engine/Shaders/BasePassPixelShader.usf b/Engine/Shaders/BasePassPixelShader.usf
--- a/Engine/Shaders/BasePassPixelShader.usf
+++ b/Engine/Shaders/BasePassPixelShader.usf
@@ -31,6 +31,10 @@
 	 
 #endif
 
+#if MATERIAL_SHADINGMODEL_HAIR
+#include "ShadingModels.usf"
+#endif
+
 Texture2D		HZBTexture;
 SamplerState	HZBSampler;
 Texture2D		PrevSceneColor;
@@ -367,6 +371,7 @@ void GetPrecomputedIndirectLightingAndSkyLight(
 	FMaterialPixelParameters MaterialParameters, 
 	FVertexFactoryInterpolantsVSToPS Interpolants,
 	FBasePassInterpolantsVSToPS BasePassInterpolants,
+	float3 DiffuseDir,
 	out float3 OutDiffuseLighting,
 	out float3 OutSubsurfaceLighting,
 	out float OutIndirectIrradiance)
@@ -404,11 +409,11 @@ void GetPrecomputedIndirectLightingAndSkyLight(
 			CachedSH.B.V = float4(Vector0.z, Vector1.z, Vector2.z, Vector2.w);
 
 			// Diffuse convolution
-			FTwoBandSHVector DiffuseTransferSH = CalcDiffuseTransferSH(MaterialParameters.WorldNormal, 1);
+			FTwoBandSHVector DiffuseTransferSH = CalcDiffuseTransferSH(DiffuseDir, 1);
 			OutDiffuseLighting = max(half3(0,0,0), DotSH(CachedSH, DiffuseTransferSH)) / PI;
 
 			#if MATERIAL_SHADINGMODEL_TWOSIDED_FOLIAGE
-				FTwoBandSHVector SubsurfaceTransferSH = CalcDiffuseTransferSH(-MaterialParameters.WorldNormal, 1);
+				FTwoBandSHVector SubsurfaceTransferSH = CalcDiffuseTransferSH(-DiffuseDir, 1);
 				OutSubsurfaceLighting += max(half3(0,0,0), DotSH(CachedSH, SubsurfaceTransferSH)) / PI;
 			#endif
 
@@ -438,12 +443,12 @@ void GetPrecomputedIndirectLightingAndSkyLight(
 			PointIndirectLighting.G.V = PrecomputedLightingBuffer.IndirectLightingSHCoefficients[1];
 			PointIndirectLighting.B.V = PrecomputedLightingBuffer.IndirectLightingSHCoefficients[2];
 
-			FTwoBandSHVector DiffuseTransferSH = CalcDiffuseTransferSH(MaterialParameters.WorldNormal, 1);
+			FTwoBandSHVector DiffuseTransferSH = CalcDiffuseTransferSH(DiffuseDir, 1);
 			// Compute diffuse lighting which takes the normal into account
 			OutDiffuseLighting = max(half3(0,0,0), DotSH(PointIndirectLighting, DiffuseTransferSH));
 
 			#if MATERIAL_SHADINGMODEL_TWOSIDED_FOLIAGE
-				FTwoBandSHVector SubsurfaceTransferSH = CalcDiffuseTransferSH(-MaterialParameters.WorldNormal, 1);
+				FTwoBandSHVector SubsurfaceTransferSH = CalcDiffuseTransferSH(-DiffuseDir, 1);
 				OutSubsurfaceLighting += max(half3(0,0,0), DotSH(PointIndirectLighting, SubsurfaceTransferSH));
 			#endif
 
@@ -455,14 +460,14 @@ void GetPrecomputedIndirectLightingAndSkyLight(
 		float2 LightmapUV0, LightmapUV1;
 		GetLightMapCoordinates(Interpolants, LightmapUV0, LightmapUV1);
 		SkyOcclusionUV = LightmapUV0;
-		GetLightMapColorHQ(LightmapUV0, LightmapUV1, MaterialParameters.WorldNormal, OutDiffuseLighting, OutSubsurfaceLighting);
+		GetLightMapColorHQ(LightmapUV0, LightmapUV1, DiffuseDir, OutDiffuseLighting, OutSubsurfaceLighting);
 
 	// Low quality texture lightmaps
 	#elif LQ_TEXTURE_LIGHTMAP
 
 		float2 LightmapUV0, LightmapUV1;
 		GetLightMapCoordinates(Interpolants, LightmapUV0, LightmapUV1);
-		OutDiffuseLighting = GetLightMapColorLQ(LightmapUV0, LightmapUV1, MaterialParameters.WorldNormal).rgb;
+		OutDiffuseLighting = GetLightMapColorLQ(LightmapUV0, LightmapUV1, DiffuseDir).rgb;
 
 	#endif
 
@@ -471,7 +476,7 @@ void GetPrecomputedIndirectLightingAndSkyLight(
 
 	float3 SkyDiffuseLighting;
 	float3 SkySubsurfaceLighting;
-	GetSkyLighting(MaterialParameters.WorldNormal, SkyOcclusionUV, SkyDiffuseLighting, SkySubsurfaceLighting);
+	GetSkyLighting(DiffuseDir, SkyOcclusionUV, SkyDiffuseLighting, SkySubsurfaceLighting);
 
 	OutSubsurfaceLighting += SkySubsurfaceLighting;
 
@@ -485,7 +490,7 @@ void GetPrecomputedIndirectLightingAndSkyLight(
 	{
 	#if (TRANSLUCENCY_LIGHTING_SURFACE || TRANSLUCENCY_LIGHTING_SURFACE_PERPIXEL) && FEATURE_LEVEL >= FEATURE_LEVEL_SM5
 		float DiffuseMip = ComputeReflectionCaptureMipFromRoughness(1);
-		float4 DiffuseIBL = TextureCubeArraySampleLevel(ReflectionCubemap, ReflectionCubemapSampler, MaterialParameters.WorldNormal, CubemapArrayIndex, DiffuseMip);
+		float4 DiffuseIBL = TextureCubeArraySampleLevel(ReflectionCubemap, ReflectionCubemapSampler, DiffuseDir, CubemapArrayIndex, DiffuseMip);
 		
 		OutDiffuseLighting += DiffuseIBL.rgb + (1 - DiffuseIBL.a) * SkyDiffuseLighting;
 	#endif
@@ -755,10 +760,7 @@ void Main(
 
 	#if !MATERIAL_SHADINGMODEL_UNLIT
 
-		float3 DiffuseIndirectLighting;
-		float3 SubsurfaceIndirectLighting;
-		GetPrecomputedIndirectLightingAndSkyLight(MaterialParameters, Interpolants, BasePassInterpolants, DiffuseIndirectLighting, SubsurfaceIndirectLighting, IndirectIrradiance);
-
+		float3 DiffuseDir = MaterialParameters.WorldNormal;
 		float3 EffectiveDiffuseColor = DiffuseColor;
 
 		#if MATERIAL_SHADINGMODEL_SUBSURFACE || MATERIAL_SHADINGMODEL_PREINTEGRATED_SKIN
@@ -767,6 +769,28 @@ void Main(
 			EffectiveDiffuseColor += SubsurfaceColor;
 		#endif
 
+		#if MATERIAL_SHADINGMODEL_HAIR
+			{
+				// TODO rearrange so that this doesn't need to be built twice
+				FGBufferData GBuffer = (FGBufferData)0;
+				GBuffer.WorldNormal = MaterialParameters.WorldNormal;
+				GBuffer.BaseColor = BaseColor;
+				GBuffer.Metallic = Metallic;
+				GBuffer.Specular = Specular;
+				GBuffer.Roughness = Roughness;
+				
+				float3 N = MaterialParameters.WorldNormal;
+				float3 V = MaterialParameters.CameraVector;
+				float3 L = normalize( V - N * dot(V,N) );
+				DiffuseDir = L;
+				EffectiveDiffuseColor = PI * HairShading( GBuffer, L, V, N, 1, 0, 0.2, uint2(0,0) );
+			}
+		#endif
+			
+		float3 DiffuseIndirectLighting;
+		float3 SubsurfaceIndirectLighting;
+		GetPrecomputedIndirectLightingAndSkyLight(MaterialParameters, Interpolants, BasePassInterpolants, DiffuseDir, DiffuseIndirectLighting, SubsurfaceIndirectLighting, IndirectIrradiance);
+		
 		Color += (DiffuseIndirectLighting * EffectiveDiffuseColor + SubsurfaceIndirectLighting * SubsurfaceColor) * MaterialAO;
 
 		#if SIMPLE_DYNAMIC_LIGHTING
@@ -863,16 +887,16 @@ void Main(
 			GBuffer.ShadingModelID = SHADINGMODELID_DEFAULT_LIT;
 		#elif MATERIAL_SHADINGMODEL_SUBSURFACE
 			GBuffer.ShadingModelID = SHADINGMODELID_SUBSURFACE;
-		GBuffer.CustomData.rgb = EncodeSubsurfaceColor(SubsurfaceColor);
-		GBuffer.CustomData.a = Opacity;
+			GBuffer.CustomData.rgb = EncodeSubsurfaceColor(SubsurfaceColor);
+			GBuffer.CustomData.a = Opacity;
 		#elif MATERIAL_SHADINGMODEL_PREINTEGRATED_SKIN
 			GBuffer.ShadingModelID = SHADINGMODELID_PREINTEGRATED_SKIN;
-		GBuffer.CustomData.rgb = EncodeSubsurfaceColor(SubsurfaceColor);
-		GBuffer.CustomData.a = Opacity;
+			GBuffer.CustomData.rgb = EncodeSubsurfaceColor(SubsurfaceColor);
+			GBuffer.CustomData.a = Opacity;
 		#elif MATERIAL_SHADINGMODEL_SUBSURFACE_PROFILE
 			GBuffer.ShadingModelID = SHADINGMODELID_SUBSURFACE_PROFILE;
-		GBuffer.CustomData.rgb = EncodeSubsurfaceProfile(SubsurfaceProfile);
-		GBuffer.CustomData.a = Opacity;
+			GBuffer.CustomData.rgb = EncodeSubsurfaceProfile(SubsurfaceProfile);
+			GBuffer.CustomData.a = Opacity;
 		#elif MATERIAL_SHADINGMODEL_CLEAR_COAT
 			GBuffer.ShadingModelID = SHADINGMODELID_CLEAR_COAT;
 
@@ -880,7 +904,7 @@ void Main(
 			float ClearCoatRoughness	= saturate( GetMaterialCustomData1(MaterialParameters) );
 			float MetalSpec = 0.9;
 
-		float NoV = saturate( dot( MaterialParameters.WorldNormal, MaterialParameters.CameraVector ) );
+			float NoV = saturate( dot( MaterialParameters.WorldNormal, MaterialParameters.CameraVector ) );
 
 			// Approximation of refraction's effect on EnvBRDF
 			float RefractionScale = ( (NoV * 0.5 + 0.5) * NoV - 1 ) * saturate( 1.25 - 1.25 * Roughness ) + 1;
diff --git a/Engine/Shaders/DeferredLightingCommon.usf b/Engine/Shaders/DeferredLightingCommon.usf
--- a/Engine/Shaders/DeferredLightingCommon.usf
+++ b/Engine/Shaders/DeferredLightingCommon.usf
@@ -322,7 +322,8 @@ float4 GetDynamicLighting(float3 WorldPosition, float3 CameraVector, float2 InUV
 			GetShadowTerms(ScreenSpaceData, LightData, WorldPosition, LightAttenuation, SurfaceShadow, SubsurfaceShadow);
 
 			// greatly reduces shadow mapping artifacts
-			SurfaceShadow *= saturate(dot(N, L) * 6 - 0.2);
+			// Commented out because it reduces character shading quality.
+			//SurfaceShadow *= saturate(dot(N, L) * 6 - 0.2);
 		}
 		else
 		{
diff --git a/Engine/Shaders/PostProcessAmbient.usf b/Engine/Shaders/PostProcessAmbient.usf
--- a/Engine/Shaders/PostProcessAmbient.usf
+++ b/Engine/Shaders/PostProcessAmbient.usf
@@ -11,6 +11,7 @@
 #include "Random.usf"
 #include "BRDF.usf"
 #include "MonteCarlo.usf"
+#include "ShadingModels.usf"
 
 #define IMPORTANCE_SAMPLE		0
 
@@ -90,154 +91,6 @@ float3 SpecularIBL( uint2 Random, float3 SpecularColor, float Roughness, float3
 	return SpecularLighting / NumSamples;
 }
 
-float3 StandardShading( FGBufferData GBuffer, float Roughness, float3 L, float3 V, half3 N )
-{
-	float3 H = normalize(V + L);
-	float NoL = saturate( dot(N, L) );
-	//float NoV = saturate( dot(N, V) );
-	float NoV = max( dot(N, V), 1e-5 );
-	float NoH = saturate( dot(N, H) );
-	float VoH = saturate( dot(V, H) );
-	
-	// Generalized microfacet specular
-	float D = D_GGX( Roughness, NoH );
-	float Vis = Vis_SmithJointApprox( Roughness, NoV, NoL );
-	float3 F = F_Schlick( GBuffer.SpecularColor, VoH );
-
-	float3 Diffuse = Diffuse_Lambert( GBuffer.DiffuseColor );
-
-	return Diffuse + (D * Vis) * F;
-}
-
-float3 ClearCoatShading( FGBufferData GBuffer, float Roughness, float3 L, float3 V, half3 N )
-{
-	const float ClearCoat			= GBuffer.CustomData.x;
-	const float ClearCoatRoughness	= GBuffer.CustomData.y;
-	const float Film = 1 * ClearCoat;
-	const float MetalSpec = 0.9;
-
-#if 0
-	float3 H = normalize(V + L);
-	float NoL = saturate( dot(N, L) );
-	float NoV = saturate( dot(N, V) );
-	float NoH = saturate( dot(N, H) );
-	float VoH = saturate( dot(V, H) );
-	
-	// Generalized microfacet specular
-	float D = D_GGX( ClearCoatRoughness, NoH );
-	float Vis = Vis_Kelemen( VoH );
-	
-	// F_Schlick
-	float F0 = 0.04;
-	float Fc = pow( 1 - VoH, 5 );
-	float F = Fc + (1 - Fc) * F0;
-	F *= ClearCoat;
-
-	float Fr1 = D * Vis * F;
-
-	float LayerAttenuation = (1 - F);
-
-	// Generalized microfacet specular
-	float D2 = D_GGX( Roughness, NoH );
-	float Vis2 = Vis_Schlick( Roughness, NoV, NoL );
-	float3 F2 = saturate( 50.0 * GBuffer.SpecularColor.g ) * Fc + (1 - Fc) * GBuffer.SpecularColor;
-
-	float3 Fr2 = Diffuse_Lambert( GBuffer.DiffuseColor ) + (D2 * Vis2) * F2;
-	
-	return Fr1 + Fr2 * LayerAttenuation;
-#else
-	float3 H = normalize(V + L);
-	float NoL = saturate( dot(N, L) );
-	float NoV = saturate( dot(N, V) );
-	float NoH = saturate( dot(N, H) );
-	float VoH = saturate( dot(V, H) );
-	
-	// Hard coded IOR of 1.5
-
-	// Generalized microfacet specular
-	float D = D_GGX( ClearCoatRoughness, NoH );
-	float Vis = Vis_Kelemen( VoH );
-	
-	// F_Schlick
-	float F0 = 0.04;
-	float Fc = pow( 1 - VoH, 5 );
-	float F = Fc + (1 - Fc) * F0;
-
-	float Fr1 = D * Vis * F;
-
-#if 1
-	// Refract rays
-	float3 L2 = refract( -L, -H, 1 / 1.5 );
-	float3 V2 = refract( -V, -H, 1 / 1.5 );
-
-	// LoH == VoH
-	float RefractBlend = sqrt( 4 * VoH*VoH + 5 ) / 3 + 2.0 / 3 * VoH;
-	//float3 L2 = RefractBlend * H - L / 1.5;
-	//float3 V2 = RefractBlend * H - V / 1.5;
-
-	float3 H2 = normalize( V2 + L2 );
-	float NoL2 = saturate( dot(N, L2) );
-	float NoV2 = saturate( dot(N, V2) );
-	float NoH2 = saturate( dot(N, H2) );
-	float VoH2 = saturate( dot(V2, H2) );
-#else
-	// Approximation
-	float RefractBlend = (0.22 * VoH + 0.7) * VoH + 0.745;	// 2 mad
-	// Dot products distribute. No need for L2 and V2.
-	float RefractNoH = RefractBlend * NoH;					// 1 mul
-	float NoL2 = saturate( RefractNoH - (1 / 1.5) * NoL );	// 1 mad
-	float NoV2 = saturate( RefractNoH - (1 / 1.5) * NoV );	// 1 mad
-	// Should refract H too but unimportant
-#endif
-
-	NoL2 = max( 0.001, NoL2 );
-	NoV2 = max( 0.001, NoV2 );
-
-	float3 AbsorptionColor = (1 - Film) + GBuffer.BaseColor * ( Film * (1 / MetalSpec) );
-	float  AbsorptionDist = rcp(NoV2) + rcp(NoL2);
-	float3 Absorption = pow( AbsorptionColor, 0.5 * AbsorptionDist );
-
-	// Approximation
-	//float  AbsorptionDist = ( NoV2 + NoL2 ) / ( NoV2 * NoL2 );
-	//float3 Absorption = AbsorptionColor * ( AbsorptionColor * (AbsorptionDist * 0.5 - 1) + (2 - 0.5 * AbsorptionDist) );
-	//float3 Absorption = AbsorptionColor + AbsorptionColor * (AbsorptionColor - 1) * (AbsorptionDist * 0.5 - 1);	// use for shared version
-	
-#if 1
-	float F21 = F_Schlick( 0.04, saturate( dot(V2, H) ) );
-
-	float k = Square( Roughness ) * 0.5;
-	float G_SchlickV2 = NoV2 / ( NoV2 * (1 - k) + k );
-	float G_SchlickL2 = NoL2 / ( NoL2 * (1 - k) + k );
-	float TotalInternalReflection = 1 - F21 * G_SchlickV2 * G_SchlickL2;
-	float3 LayerAttenuation = ( (1 - F) * TotalInternalReflection ) * Absorption;
-#else
-	// Approximation
-	float3 LayerAttenuation = (1 - F) * Absorption;
-#endif
-
-	// Approximation for IOR == 1.5
-	//SpecularColor = ChangeBaseMedium( SpecularColor, 1.5 );
-	//SpecularColor = saturate( ( 0.55 * SpecularColor + (0.45 * 1.08) ) * SpecularColor - (0.45 * 0.08) );
-	// Treat SpecularColor as relative to IOR. Artist compensates.
-	
-	// Generalized microfacet specular
-#if 1
-	float D2 = D_GGX( Roughness, NoH2 );
-	float Vis2 = Vis_Schlick( Roughness, NoV2, NoL2 );
-	float3 F2 = F_Schlick( MetalSpec, VoH2 );
-#else
-	// Approximation
-	float D2 = D_GGX( Roughness, NoH );
-	float Vis2 = Vis_Schlick( Roughness, NoV2, NoL2 );
-	float3 F2 = F_Schlick( GBuffer.SpecularColor, VoH );
-#endif
-
-	float3 Fr2 = Diffuse_Lambert( GBuffer.DiffuseColor ) + (D2 * Vis2) * F2;
-	
-	return Fr1 + Fr2 * LayerAttenuation;
-#endif
-}
-
 
 float3 ImageBasedLightingMIS( FGBufferData GBuffer, float3 V, float3 N, uint2 Random )
 {
@@ -305,26 +158,28 @@ float3 ImageBasedLightingMIS( FGBufferData GBuffer, float3 V, float3 N, uint2 Ra
 #if USE_CLEARCOAT
 				BRANCH if( GBuffer.ShadingModelID == SHADINGMODELID_CLEAR_COAT )
 				{
-					Shading = ClearCoatShading( GBuffer, GBuffer.Roughness, L, V, N );
+					Shading = ClearCoatShading( GBuffer, GBuffer.Roughness, 1, L, V, N );
 				}
 				else
 #endif
-#if USE_CLOTH
+//#if USE_CLOTH
 				BRANCH if( GBuffer.ShadingModelID == SHADINGMODELID_CLOTH )
 				{
-					Shading = ClothShading( GBuffer, GBuffer.Roughness, 1, L, V, N, 1 );
+					Shading = ClothShading( GBuffer, GBuffer.Roughness, 1, L, V, N );
 				}
 				else
-#endif
-#if USE_EYE
+//#endif
+//#if USE_EYE
 				BRANCH if( GBuffer.ShadingModelID == SHADINGMODELID_EYE )
 				{
-					Shading = EyeShading( GBuffer, GBuffer.Roughness, 1, L, V, N, 1 );
+					Shading  = EyeShading( GBuffer, GBuffer.Roughness, 1, L, V, N ) * NoL;
+					Shading += EyeSubsurfaceShading( GBuffer, L, V, N );
+					NoL = 1;
 				}
 				else
-#endif
+//#endif
 				{
-					Shading = StandardShading( GBuffer, GBuffer.Roughness, L, V, N );
+					Shading = StandardShading( GBuffer, GBuffer.Roughness, 1, L, V, N, 1 );
 				}
 			
 				Lighting += SampleColor * Shading * ( NoL * Weight );
@@ -458,6 +313,37 @@ float3 ApproximateSpecularIBL( uint2 Random, float3 SpecularColor, float Roughne
 	return PrefilteredColor * ( SpecularColor * AB.x + AB.y );
 }
 
+
+float3 ImageBasedLightingHair( FGBufferData GBuffer, float3 V, float3 N, uint2 Random )
+{
+	float3 Lighting = 0;
+
+	uint NumSamples = 32;
+	
+	LOOP
+	for( uint i = 0; i < NumSamples; i++ )
+	{
+		float2 E = Hammersley( i, NumSamples, Random );
+		float3 L = UniformSampleSphere( E ).xyz;
+		
+		{
+			float3 SampleColor = AmbientCubemap.SampleLevel( AmbientCubemapSampler, L, 0 ).rgb;
+
+			float PDF = 1.0 / (4 * PI);
+
+			float InvWeight = PDF * NumSamples;
+			float Weight = rcp( InvWeight );
+
+			float3 Shading = 0;
+			Shading = HairShading( GBuffer, L, V, N, 1, 0, 0, Random );
+			
+			Lighting += SampleColor * Shading * Weight;
+		}
+	}
+
+	return Lighting;
+}
+
 void MainPS(in noperspective float4 UVAndScreenPos : TEXCOORD0, float4 SvPosition : SV_POSITION, out float4 OutColor : SV_Target0)
 {
 	float2 UV = UVAndScreenPos.xy;
@@ -541,15 +427,32 @@ void MainPS(in noperspective float4 UVAndScreenPos : TEXCOORD0, float4 SvPositio
 		SpecularContribution += SampleColor * F;
 	}
 #endif
+	
+	BRANCH if( GBuffer.ShadingModelID == SHADINGMODELID_HAIR )
+	{
+		float3 FakeNormal = normalize( V - N * dot(V,N) );
+		FakeNormal = normalize( FakeNormal + 0 * float3(0,0,1) );
+		SpecularContribution = TextureCubeSampleLevel(AmbientCubemap, AmbientCubemapSampler, FakeNormal, AbsoluteDiffuseMip).rgb;
+		SpecularContribution *= PI * HairShading( GBuffer, FakeNormal, V, N, 1, 0, 0.2, Random );
+		NonSpecularContribution = 0;
+	}
 
 #if IMPORTANCE_SAMPLE
-	if( GBuffer.ShadingModelID > 0 )
+	if( GBuffer.ShadingModelID > 0 && UV.x > 0.5 )
 	{
 		NonSpecularContribution = 0;
-		SpecularContribution = ImageBasedLightingMIS( GBuffer, -ScreenVector, GBuffer.WorldNormal, Random );
+
+		BRANCH if( GBuffer.ShadingModelID == SHADINGMODELID_HAIR )
+		{
+			SpecularContribution = ImageBasedLightingHair( GBuffer, -ScreenVector, GBuffer.WorldNormal, Random );
+		}
+		else
+		{
+			SpecularContribution = ImageBasedLightingMIS( GBuffer, -ScreenVector, GBuffer.WorldNormal, Random );
+		}
 	}
 #endif
-	
+
 	// apply darkening from ambient occlusion (does not use PostprocessInput1 to set white texture if SSAO is off)
 	float AmbientOcclusion = GBuffer.GBufferAO * ScreenSpaceData.AmbientOcclusion;
 
diff --git a/Engine/Shaders/PostProcessAmbientOcclusion.usf b/Engine/Shaders/PostProcessAmbientOcclusion.usf
--- a/Engine/Shaders/PostProcessAmbientOcclusion.usf
+++ b/Engine/Shaders/PostProcessAmbientOcclusion.usf
@@ -526,6 +526,13 @@ void MainPSandCS(in float4 UVAndScreenPos, float4 SvPosition, out float4 OutColo
 	float SceneDepth = GetDepthFromAOInput(UV);
 	float3 WorldNormal = GetWorldSpaceNormalFromAOInput(UV);
 
+	FGBufferData GBuffer = GetGBufferData(UV);
+	BRANCH if( GBuffer.ShadingModelID == SHADINGMODELID_HAIR )
+	{
+		float3 V = -normalize(mul(float4(UVAndScreenPos.zw, 1, 0), View.ScreenToWorld).xyz);
+		WorldNormal = normalize( V - WorldNormal * dot(V, WorldNormal) );
+	}
+
 	float3 ViewSpaceNormal = normalize(mul(WorldNormal, (float3x3)View.TranslatedWorldToView));
 	
 	float3 ViewSpacePosition = ReconstructCSPos(SceneDepth, ScreenPos);
diff --git a/Engine/Shaders/PostProcessTemporalCommon.usf b/Engine/Shaders/PostProcessTemporalCommon.usf
--- a/Engine/Shaders/PostProcessTemporalCommon.usf
+++ b/Engine/Shaders/PostProcessTemporalCommon.usf
@@ -484,7 +484,7 @@ which can be configured via defines for various temporal AA passes.
 		float HistoryFactor = DistToClamp * HistoryAmount * (1.0 + HistoryBlur * HistoryAmount * 8.0);
 	#endif
 	float BlendFinal = saturate(HistoryFactor / (DistToClamp + LumaMax - LumaMin));
-	BlendFinal = 0.08;
+	BlendFinal = 0.04;
 	#if RESPONSIVE
 		// Responsive forces 1/4 of new frame.
 		BlendFinal = 1.0/4.0;
diff --git a/Engine/Shaders/ReflectionEnvironmentComputeShaders.usf b/Engine/Shaders/ReflectionEnvironmentComputeShaders.usf
--- a/Engine/Shaders/ReflectionEnvironmentComputeShaders.usf
+++ b/Engine/Shaders/ReflectionEnvironmentComputeShaders.usf
@@ -626,7 +626,7 @@ void ReflectionEnvironmentTiledDeferredMain(
 #endif
 
 	BRANCH
-	if( GBuffer.ShadingModelID != SHADINGMODELID_UNLIT )
+	if( GBuffer.ShadingModelID != SHADINGMODELID_UNLIT && GBuffer.ShadingModelID != SHADINGMODELID_HAIR )
 	{
 		float3 N = GBuffer.WorldNormal;
 		float3 V = -CameraToPixel;
diff --git a/Engine/Shaders/ReflectionEnvironmentShaders.usf b/Engine/Shaders/ReflectionEnvironmentShaders.usf
--- a/Engine/Shaders/ReflectionEnvironmentShaders.usf
+++ b/Engine/Shaders/ReflectionEnvironmentShaders.usf
@@ -459,7 +459,7 @@ void ReflectionApplyPS(
 	
 	BRANCH
 	// Only light pixels marked as lit
-	if( GBuffer.ShadingModelID > 0 )
+	if( GBuffer.ShadingModelID > 0 && GBuffer.ShadingModelID != SHADINGMODELID_HAIR )
 	{
 		// TODO if SpecularColor and Roughness were in same texture would remove 2 fetches
 		float3 V = -normalize( mul( float4(UVAndScreenPos.zw, 1, 0), View.ScreenToWorld ).xyz );
diff --git a/Engine/Shaders/ShadingModels.usf b/Engine/Shaders/ShadingModels.usf
--- a/Engine/Shaders/ShadingModels.usf
+++ b/Engine/Shaders/ShadingModels.usf
@@ -517,10 +517,10 @@ float3 HairShadingRef( FGBufferData GBuffer, float3 L, float3 V, half3 N, uint2
 // Approximation to HairShadingRef using concepts from the following papers:
 // [Marschner et al. 2003, "Light Scattering from Human Hair Fibers"]
 // [Pekelis et al. 2015, "A Data-Driven Light Scattering Model for Hair"]
-float3 HairShading( FGBufferData GBuffer, float3 L, float3 V, half3 N, float Shadow, uint2 Random )
+float3 HairShading( FGBufferData GBuffer, float3 L, float3 V, half3 N, float Shadow, float Backlit, float Area, uint2 Random )
 {
 	//const float3 DiffuseN	= OctahedronToUnitVector( GBuffer.CustomData.xy * 2 - 1 );
-	const float Backlit	= GBuffer.CustomData.z;
+	//const float Backlit	= GBuffer.CustomData.z;
 
 #if HAIR_REFERENCE
 	float3 S = HairShadingRef( GBuffer, L, V, N, Random );
@@ -531,7 +531,9 @@ float3 HairShading( FGBufferData GBuffer, float3 L, float3 V, half3 N, float Sha
 	const float VoL       = dot(V,L);
 	const float SinThetaL = dot(N,L);
 	const float SinThetaV = dot(N,V);
-	const float CosThetaD = cos( 0.5 * abs( asin( SinThetaV ) - asin( SinThetaL ) ) );
+	float CosThetaD = cos( 0.5 * abs( asin( SinThetaV ) - asin( SinThetaL ) ) );
+
+	//CosThetaD = abs( CosThetaD ) < 0.01 ? 0.01 : CosThetaD;
 
 	const float3 Lp = L - SinThetaL * N;
 	const float3 Vp = V - SinThetaV * N;
@@ -543,18 +545,18 @@ float3 HairShading( FGBufferData GBuffer, float3 L, float3 V, half3 N, float Sha
 	//float n_prime = sqrt( n*n - 1 + Pow2( CosThetaD ) ) / CosThetaD;
 	float n_prime = 1.2 / CosThetaD + 0.35 * CosThetaD;
 
-	float Shift = 0.05;
+	float Shift = 0.035;
 	float Alpha[] =
 	{
+		-Shift * 2,
 		Shift,
-		-Shift * 0.5,
-		-Shift * 1.5,
+		Shift * 4,
 	};
 	float B[] =
 	{
-		Pow2( GBuffer.Roughness ),
-		Pow2( GBuffer.Roughness ) / 2,
-		Pow2( GBuffer.Roughness ) * 2,
+		Area + Pow2( GBuffer.Roughness ),
+		Area + Pow2( GBuffer.Roughness ) / 2,
+		Area + Pow2( GBuffer.Roughness ) * 2,
 	};
 
 	float3 S = 0;
@@ -567,7 +569,7 @@ float3 HairShading( FGBufferData GBuffer, float3 L, float3 V, half3 N, float Sha
 
 		float Mp = Hair_g( B[0] * sqrt(2.0) * CosHalfPhi, SinThetaL + SinThetaV - Shift );
 		float Np = 0.25 * CosHalfPhi;
-		float Fp = Hair_F( sqrt( 0.5 + 0.5 * VoL ) );
+		float Fp = Hair_F( sqrt( saturate( 0.5 + 0.5 * VoL ) ) );
 		S += Mp * Np * Fp * ( GBuffer.Specular * 2 ) * lerp( 1, Backlit, saturate(-VoL) );
 	}
 
@@ -625,7 +627,7 @@ float3 HairShading( FGBufferData GBuffer, float3 L, float3 V, half3 N, float Sha
 	float DiffuseScatter = (1 / PI) * NoL * GBuffer.Metallic;
 	float Luma = Luminance( GBuffer.BaseColor );
 	float3 ScatterTint = pow( GBuffer.BaseColor / Luma, 1 - Shadow );
-	S += lerp( GBuffer.BaseColor, Luma, 0.25 ) * DiffuseScatter * ScatterTint;
+	S += sqrt( GBuffer.BaseColor ) * DiffuseScatter * ScatterTint;
 
 	return S;
 }
@@ -740,7 +742,6 @@ float3 SurfaceShading( FGBufferData GBuffer, float3 LobeRoughness, float3 LobeEn
 		case SHADINGMODELID_TWOSIDED_FOLIAGE:
 			return StandardShading( GBuffer, LobeRoughness, LobeEnergy, L, V, N, DiffSpecMask);
 		case SHADINGMODELID_CLEAR_COAT:
-			// this path does not support DiffSpecMask yet
 			return ClearCoatShading( GBuffer, LobeRoughness, LobeEnergy, L, V, N );
 		case SHADINGMODELID_CLOTH:
 			return ClothShading( GBuffer, LobeRoughness, LobeEnergy, L, V, N );
@@ -764,7 +765,7 @@ float3 SubsurfaceShading( FGBufferData GBuffer, float3 L, float3 V, half3 N, flo
 		case SHADINGMODELID_TWOSIDED_FOLIAGE:
 			return SubsurfaceShadingTwoSided( SubsurfaceColor, L, V, N );
 		case SHADINGMODELID_HAIR:
-			return HairShading( GBuffer, L, V, N, Shadow, Random );
+			return HairShading( GBuffer, L, V, N, Shadow, 1, 0, Random );
 		case SHADINGMODELID_EYE:
 			return EyeSubsurfaceShading( GBuffer, L, V, N );
 		default:
diff --git a/Engine/Shaders/ShadowProjectionPixelShader.usf b/Engine/Shaders/ShadowProjectionPixelShader.usf
--- a/Engine/Shaders/ShadowProjectionPixelShader.usf
+++ b/Engine/Shaders/ShadowProjectionPixelShader.usf
@@ -107,9 +107,14 @@ void Main(
 		BRANCH
 		if (IsSubsurfaceModel(InGBufferData.ShadingModelID))
 		{
-			float Opacity = InGBufferData.ShadingModelID == SHADINGMODELID_HAIR || InGBufferData.ShadingModelID == SHADINGMODELID_EYE ? 1 : InGBufferData.CustomData.a;
+			float Opacity = InGBufferData.CustomData.a;
 			// Derive density from a heuristic using opacity, tweaked for useful falloff ranges and to give a linear depth falloff with opacity
 			float Density = -.05f * log(1 - min(Opacity, .999f));
+			if( InGBufferData.ShadingModelID == SHADINGMODELID_HAIR || InGBufferData.ShadingModelID == SHADINGMODELID_EYE )
+			{
+				Opacity = 1;
+				Density = 1;
+			}
 			//@todo - Use a larger kernel filter based on material opacity
 			// This is not being done currently because a varying opacity causes discontinuities in the shadow result
 			// Note: Snapping based on the largest size kernel is done in FProjectedShadowInfo::FProjectedShadowInfo, MaxDownsampleFactor is set to 4
diff --git a/Engine/Shaders/SkyLighting.usf b/Engine/Shaders/SkyLighting.usf
--- a/Engine/Shaders/SkyLighting.usf
+++ b/Engine/Shaders/SkyLighting.usf
@@ -8,6 +8,7 @@
 #include "DeferredShadingCommon.usf"
 #include "ReflectionEnvironmentShared.usf"
 #include "BRDF.usf"
+#include "ShadingModels.usf"
 
 Texture2D BentNormalAOTexture;
 SamplerState BentNormalAOSampler;
@@ -79,6 +80,17 @@ void SkyLightDiffusePS(
 			Lighting += ScalarFactors * SubsurfaceLookup * SubsurfaceColor;
 		}
 
+		float3 DiffuseColor = GBuffer.DiffuseColor;
+
+		BRANCH if( GBuffer.ShadingModelID == SHADINGMODELID_HAIR )
+		{
+			float3 N = GBuffer.WorldNormal;
+			float3 V = -normalize(mul(float4(UVAndScreenPos.zw, 1, 0), View.ScreenToWorld).xyz);
+			float3 L = normalize( V - N * dot(V,N) );
+			SkyLightingNormal = L;
+			DiffuseColor = PI * HairShading( GBuffer, L, V, N, 1, 0, 0.2, uint2(0,0) );
+		}
+
 		// Compute the preconvolved incoming lighting with the bent normal direction
 		float3 DiffuseLookup = GetSkySHDiffuse(SkyLightingNormal) * View.SkyLightColor.rgb;
 
@@ -87,7 +99,7 @@ void SkyLightDiffusePS(
 		DiffuseLookup *= DirectionalOcclusion;
 
 		// Apply AO to the sky diffuse and account for darkening due to the geometry term
-		Lighting += (ScalarFactors * DotProductFactor) * DiffuseLookup * GBuffer.DiffuseColor + (1 - SkyVisibility) * OcclusionTintAndMinOcclusion.xyz;
+		Lighting += (ScalarFactors * DotProductFactor) * DiffuseLookup * DiffuseColor + (1 - SkyVisibility) * OcclusionTintAndMinOcclusion.xyz;
 
 		Lighting += DiffuseIrradiance * GBuffer.DiffuseColor * (GBuffer.GBufferAO * ScreenSpaceData.AmbientOcclusion);
 
diff --git a/Engine/Shaders/SlateElementPixelShader.usf b/Engine/Shaders/SlateElementPixelShader.usf
--- a/Engine/Shaders/SlateElementPixelShader.usf
+++ b/Engine/Shaders/SlateElementPixelShader.usf
@@ -310,7 +310,7 @@ float4 Main( VertexToPixelInterpolants VIn ) : SV_Target0
 #endif
 
 
-#if !PRE_MULTIPLIED_ALPHA && !MATERIALBLENDING_ADDITIVE && !MATERIALBLENDING_TRANSLUCENT
+#if !PRE_MULTIPLIED_ALPHA && !USE_MATERIALS
 	OutColor.rgb = OutColor.rgb * OutColor.a;
 #endif
 
diff --git a/Engine/Source/Developer/FriendsAndChat/Private/FriendsAndChatStyle.cpp b/Engine/Source/Developer/FriendsAndChat/Private/FriendsAndChatStyle.cpp
--- a/Engine/Source/Developer/FriendsAndChat/Private/FriendsAndChatStyle.cpp
+++ b/Engine/Source/Developer/FriendsAndChat/Private/FriendsAndChatStyle.cpp
@@ -58,13 +58,6 @@ FFriendsAndChatStyle& FFriendsAndChatStyle::SetFriendsMarkupStyle(const FFriends
 	return *this;
 }
 
-FFriendsAndChatStyle& FFriendsAndChatStyle::SetBorderPadding(const FMargin& Padding)
-{
-	BorderPadding = Padding;
-	return *this;
-}
-
-
 FFriendsAndChatStyle& FFriendsAndChatStyle::SetScrollbarStyle(const FScrollBarStyle& InScrollBarStyle)
 {
 	ScrollBarStyle = InScrollBarStyle;
diff --git a/Engine/Source/Developer/FriendsAndChat/Private/FriendsListStyle.cpp b/Engine/Source/Developer/FriendsAndChat/Private/FriendsListStyle.cpp
--- a/Engine/Source/Developer/FriendsAndChat/Private/FriendsListStyle.cpp
+++ b/Engine/Source/Developer/FriendsAndChat/Private/FriendsListStyle.cpp
@@ -74,20 +74,6 @@ FFriendsListStyle& FFriendsListStyle::SetFriendListActionButtonStyle(const FButt
 	return *this;
 }
 
-/** Friends List Close button style */
-FFriendsListStyle& FFriendsListStyle::SetFriendsListClosedButtonStyle(const FButtonStyle& ButtonStyle)
-{
-	FriendListCloseButtonStyle = ButtonStyle;
-	return *this;
-}
-
-/** Add Friend Close button style */
-FFriendsListStyle& FFriendsListStyle::SetAddFriendCloseButtonStyle(const FButtonStyle& ButtonStyle)
-{
-	AddFriendCloseButtonStyle = ButtonStyle;
-	return *this;
-}
-
 /** Optional content for the Add Friend button */
 FFriendsListStyle& FFriendsListStyle::SetAddFriendButtonContentBrush(const FSlateBrush& BrushStyle)
 {
@@ -187,39 +173,39 @@ FFriendsListStyle& FFriendsListStyle::SetUserPresenceImageSize(const FVector2D&
 	return *this;
 }
 
-FFriendsListStyle& FFriendsListStyle::SetAddFriendButtonSize(const FVector2D& Value)
+FFriendsListStyle& FFriendsListStyle::SetBackBrush(const FSlateBrush& Brush)
 {
-	AddFriendButtonSize = Value;
+	BackBrush = Brush;
 	return *this;
 }
 
-FFriendsListStyle& FFriendsListStyle::SetUserHeaderPadding(const FMargin& Margin)
+FFriendsListStyle& FFriendsListStyle::SetSelectedOptionBrush(const FSlateBrush& Brush)
 {
-	UserHeaderPadding = Margin;
+	SelectedOptionBrush = Brush;
 	return *this;
 }
 
-FFriendsListStyle& FFriendsListStyle::SetBackBrush(const FSlateBrush& Brush)
+FFriendsListStyle& FFriendsListStyle::SetSettingsBrush(const FSlateBrush& Brush)
 {
-	BackBrush = Brush;
+	SettingsBrush = Brush;
 	return *this;
 }
 
-FFriendsListStyle& FFriendsListStyle::SetSelectedOptionBrush(const FSlateBrush& Brush)
+FFriendsListStyle& FFriendsListStyle::SetSeperatorBrush(const FSlateBrush& Brush)
 {
-	SelectedOptionBrush = Brush;
+	SeperatorBrush = Brush;
 	return *this;
 }
 
-FFriendsListStyle& FFriendsListStyle::SetFilterBrush(const FSlateBrush& Brush)
+FFriendsListStyle& FFriendsListStyle::SetFontSizeBrush(const FSlateBrush& Brush)
 {
-	FilterBrush = Brush;
+	FontSizeBrush = Brush;
 	return *this;
 }
 
-FFriendsListStyle& FFriendsListStyle::SetSettingsBrush(const FSlateBrush& Brush)
+FFriendsListStyle& FFriendsListStyle::SetSearchBrush(const FSlateBrush& Brush)
 {
-	SettingsBrush = Brush;
+	SearchBrush = Brush;
 	return *this;
 }
 
diff --git a/Engine/Source/Developer/FriendsAndChat/Public/FriendsAndChatStyle.h b/Engine/Source/Developer/FriendsAndChat/Public/FriendsAndChatStyle.h
--- a/Engine/Source/Developer/FriendsAndChat/Public/FriendsAndChatStyle.h
+++ b/Engine/Source/Developer/FriendsAndChat/Public/FriendsAndChatStyle.h
@@ -48,10 +48,6 @@ struct FRIENDSANDCHAT_API FFriendsAndChatStyle
 
 	// Common Style
 
-	UPROPERTY(EditAnywhere, Category = Appearance) //TODO Antony.Carter Do we not need more of these per widget?
-	FMargin BorderPadding;
-	FFriendsAndChatStyle& SetBorderPadding(const FMargin& InBorderPadding);
-
 	UPROPERTY(EditAnywhere, Category = Appearance)
 	FScrollBarStyle ScrollBarStyle;
 	FFriendsAndChatStyle& SetScrollbarStyle(const FScrollBarStyle& InScrollBarStyle);
diff --git a/Engine/Source/Developer/FriendsAndChat/Public/FriendsListStyle.h b/Engine/Source/Developer/FriendsAndChat/Public/FriendsListStyle.h
--- a/Engine/Source/Developer/FriendsAndChat/Public/FriendsListStyle.h
+++ b/Engine/Source/Developer/FriendsAndChat/Public/FriendsListStyle.h
@@ -84,16 +84,6 @@ struct FRIENDSANDCHAT_API FFriendsListStyle
 	FButtonStyle FriendListActionButtonStyle;
 	FFriendsListStyle& SetFriendListActionButtonStyle(const FButtonStyle& ButtonStyle);
 
-	/** Friends List Close button style */
-	UPROPERTY( EditAnywhere, Category = Appearance )
-	FButtonStyle FriendListCloseButtonStyle;
-	FFriendsListStyle& SetFriendsListClosedButtonStyle(const FButtonStyle& ButtonStyle);
-
-	/** Add Friend Close button style */
-	UPROPERTY( EditAnywhere, Category = Appearance )
-	FButtonStyle AddFriendCloseButtonStyle;
-	FFriendsListStyle& SetAddFriendCloseButtonStyle(const FButtonStyle& ButtonStyle);
-
 	/** Optional content for the Add Friend button */
 	UPROPERTY( EditAnywhere, Category = Appearance )
 	FSlateBrush AddFriendButtonContentBrush;
@@ -167,14 +157,6 @@ struct FRIENDSANDCHAT_API FFriendsListStyle
 	FVector2D UserPresenceImageSize;
 	FFriendsListStyle& SetUserPresenceImageSize(const FVector2D& InUserPresenceImageSize);
 
-	UPROPERTY(EditAnywhere, Category = Appearance)
-	FVector2D AddFriendButtonSize;
-	FFriendsListStyle& SetAddFriendButtonSize(const FVector2D& Value);
-
-	UPROPERTY(EditAnywhere, Category = Appearance)
-	FMargin UserHeaderPadding;
-	FFriendsListStyle& SetUserHeaderPadding(const FMargin& InUserHeaderPadding);
-
 	UPROPERTY(EditAnywhere, Category = Appearance)
 	FSlateBrush BackBrush;
 	FFriendsListStyle& SetBackBrush(const FSlateBrush& Brush);
@@ -183,14 +165,22 @@ struct FRIENDSANDCHAT_API FFriendsListStyle
 	FSlateBrush SelectedOptionBrush;
 	FFriendsListStyle& SetSelectedOptionBrush(const FSlateBrush& Brush);
 
-	UPROPERTY(EditAnywhere, Category = Appearance)
-	FSlateBrush FilterBrush;
-	FFriendsListStyle& SetFilterBrush(const FSlateBrush& Brush);
-
 	UPROPERTY(EditAnywhere, Category = Appearance)
 	FSlateBrush SettingsBrush;
 	FFriendsListStyle& SetSettingsBrush(const FSlateBrush& Brush);
 
+	UPROPERTY(EditAnywhere, Category = Appearance)
+	FSlateBrush SeperatorBrush;
+	FFriendsListStyle& SetSeperatorBrush(const FSlateBrush& Brush);
+
+	UPROPERTY(EditAnywhere, Category = Appearance)
+	FSlateBrush FontSizeBrush;
+	FFriendsListStyle& SetFontSizeBrush(const FSlateBrush& Brush);
+
+	UPROPERTY(EditAnywhere, Category = Appearance)
+	FSlateBrush SearchBrush;
+	FFriendsListStyle& SetSearchBrush(const FSlateBrush& Brush);
+
 	UPROPERTY(EditAnywhere, Category = Appearance)
 	FMargin BackButtonMargin;
 
@@ -233,6 +223,34 @@ struct FRIENDSANDCHAT_API FFriendsListStyle
 	UPROPERTY(EditAnywhere, Category = Appearance)
 	FMargin ConfirmationButtonContentMargin;
 
+	UPROPERTY(EditAnywhere, Category = Appearance)
+	FMargin SubMenuBackIconMargin;
+
+	UPROPERTY(EditAnywhere, Category = Appearance)
+	FMargin SubMenuPageIconMargin;
+
+	UPROPERTY(EditAnywhere, Category = Appearance)
+	FMargin RadioSettingTitleMargin;
+
+	UPROPERTY(EditAnywhere, Category = Appearance)
+	FMargin SubMenuSearchIconMargin;
+
+	UPROPERTY(EditAnywhere, Category = Appearance)
+	FMargin SubMenuSearchTextMargin;
+
+	UPROPERTY(EditAnywhere, Category = Appearance)
+	FMargin SubMenuBackButtonMargin;
+
+	UPROPERTY(EditAnywhere, Category = Appearance)
+	FMargin SubMenuSettingButtonMargin;
+
+	UPROPERTY(EditAnywhere, Category = Appearance)
+	FMargin SubMenuListMargin;
+
+	UPROPERTY(EditAnywhere, Category = Appearance)
+	float SubMenuSeperatorThickness;
+
+
 
 // Clan Settings
 
diff --git a/Engine/Source/Developer/GameplayDebugger/Private/GameplayDebuggingReplicator.cpp b/Engine/Source/Developer/GameplayDebugger/Private/GameplayDebuggingReplicator.cpp
--- a/Engine/Source/Developer/GameplayDebugger/Private/GameplayDebuggingReplicator.cpp
+++ b/Engine/Source/Developer/GameplayDebugger/Private/GameplayDebuggingReplicator.cpp
@@ -122,6 +122,9 @@ void AGameplayDebuggingReplicator::PostEditChangeProperty(FPropertyChangedEvent&
 	Super::PostEditChangeProperty(PropertyChangedEvent);
 
 #if !(UE_BUILD_SHIPPING || UE_BUILD_TEST)
+	// @note patching up OR-9814
+	UGameplayDebuggingComponent* DummyPointer = GetDebugComponent();
+
 	if (!PropertyChangedEvent.Property)
 	{
 		return;
@@ -133,7 +136,7 @@ void AGameplayDebuggingReplicator::PostEditChangeProperty(FPropertyChangedEvent&
 	if (PropertyChangedEvent.Property->GetFName() == GET_MEMBER_NAME_CHECKED(AGameplayDebuggingReplicator, __FlagName_)) \
 	{ \
 		__FlagName_ ? Settings.SetFlag(EAIDebugDrawDataView::__FlagName_) : Settings.ClearFlag(EAIDebugDrawDataView::__FlagName_); \
-		GetDebugComponent()->ServerReplicateData(Settings.CheckFlag(EAIDebugDrawDataView::__FlagName_) ? EDebugComponentMessage::ActivateDataView : EDebugComponentMessage::DeactivateDataView, EAIDebugDrawDataView::__FlagName_); \
+		DebugComponent->ServerReplicateData(Settings.CheckFlag(EAIDebugDrawDataView::__FlagName_) ? EDebugComponentMessage::ActivateDataView : EDebugComponentMessage::DeactivateDataView, EAIDebugDrawDataView::__FlagName_); \
 	}else
 
 	CHECK_AND_UPDATE_FLAGS(OverHead)
@@ -150,14 +153,14 @@ void AGameplayDebuggingReplicator::PostEditChangeProperty(FPropertyChangedEvent&
 #if WITH_EQS
 	if (PropertyChangedEvent.Property->GetFName() == GET_MEMBER_NAME_CHECKED(AGameplayDebuggingReplicator, EQS))
 	{
-		GetDebugComponent()->EnableClientEQSSceneProxy(EQS);
-		GetDebugComponent()->SetEQSIndex(ActiveEQSIndex);
-		GetDebugComponent()->MarkRenderStateDirty();
+		DebugComponent->EnableClientEQSSceneProxy(EQS);
+		DebugComponent->SetEQSIndex(ActiveEQSIndex);
+		DebugComponent->MarkRenderStateDirty();
 	}
 
 	if (PropertyChangedEvent.Property->GetFName() == GET_MEMBER_NAME_CHECKED(AGameplayDebuggingReplicator, ActiveEQSIndex))
 	{
-		GetDebugComponent()->SetEQSIndex(ActiveEQSIndex);
+		DebugComponent->SetEQSIndex(ActiveEQSIndex);
 	}
 #endif // WITH_EQS
 #undef CHECK_AND_UPDATE_FLAGS
@@ -185,7 +188,12 @@ void AGameplayDebuggingReplicator::BeginPlay()
 				DebugComponentClass = UGameplayDebuggingComponent::StaticClass();
 			}
 		}
-		GetDebugComponent();
+		// @note patching up OR-9814
+		const UGameplayDebuggingComponent* DummyPointer = GetDebugComponent();
+		if (DummyPointer)
+		{
+			UE_LOG(LogGameplayDebugger, Error, TEXT("Unable to create UGameplayDebuggingComponent instance!"));
+		}
 	}
 
 	if (GetWorld() && GetNetMode() != ENetMode::NM_DedicatedServer)
@@ -360,9 +368,11 @@ bool AGameplayDebuggingReplicator::ClientEnableTargetSelection_Validate(bool, AP
 void AGameplayDebuggingReplicator::ClientEnableTargetSelection_Implementation(bool bEnable, APlayerController* Context)
 {
 #if !(UE_BUILD_SHIPPING || UE_BUILD_TEST)
-	if (GetDebugComponent())
+	// @note patching up OR-9814
+	const UGameplayDebuggingComponent* DummyPointer = GetDebugComponent();
+	if (DebugComponent)
 	{
-		GetDebugComponent()->ClientEnableTargetSelection(bEnable);
+		DebugComponent->ClientEnableTargetSelection(bEnable);
 	}
 #endif
 }
@@ -391,9 +401,11 @@ void AGameplayDebuggingReplicator::ServerReplicateMessage_Implementation(class
 		MarkComponentsRenderStateDirty();
 	}
 
-	if (GetDebugComponent())
+	// @note patching up OR-9814
+	const UGameplayDebuggingComponent* DummyPointer = GetDebugComponent();
+	if (DebugComponent)
 	{
-		GetDebugComponent()->ServerReplicateData((EDebugComponentMessage::Type)InMessage, (EAIDebugDrawDataView::Type)DataView);
+		DebugComponent->ServerReplicateData((EDebugComponentMessage::Type)InMessage, (EAIDebugDrawDataView::Type)DataView);
 	}
 #endif //!(UE_BUILD_SHIPPING || UE_BUILD_TEST)
 }
@@ -457,15 +469,17 @@ void AGameplayDebuggingReplicator::CreateTool()
 void AGameplayDebuggingReplicator::EnableTool()
 {
 #if !(UE_BUILD_SHIPPING || UE_BUILD_TEST)
-	if (GetWorld() && GetNetMode() != ENetMode::NM_DedicatedServer)
+	// @note patching up OR-9814
+	if (GetWorld() && GetNetMode() != ENetMode::NM_DedicatedServer && DebugComponent != nullptr)
 	{
-		UGameplayDebuggingControllerComponent*  GDC = FindComponentByClass<UGameplayDebuggingControllerComponent>();
+		UGameplayDebuggingControllerComponent* GDC = FindComponentByClass<UGameplayDebuggingControllerComponent>();
 		if (GDC)
 		{
 			// simulate key press
 			GDC->OnActivationKeyPressed();
 			ClientEnableTargetSelection(true, LocalPlayerOwner);
-			GetDebugComponent()->SelectTargetToDebug();
+			// @note patching up OR-9814
+			DebugComponent->SelectTargetToDebug();
 			GDC->OnActivationKeyReleased();
 		}
 	}
@@ -526,9 +540,11 @@ void AGameplayDebuggingReplicator::ServerSetActorToDebug_Implementation(AActor*
 		}
 	}
 
-	if (UGameplayDebuggingComponent* DebuggingComponent = GetDebugComponent())
+	// @note patching up OR-9814
+	const UGameplayDebuggingComponent* DummyPointer = GetDebugComponent();
+	if (DebugComponent)
 	{
-		DebuggingComponent->SetActorToDebug(InActor);
+		DebugComponent->SetActorToDebug(InActor);
 	}
 #endif
 }
@@ -656,8 +672,13 @@ void AGameplayDebuggingReplicator::DrawDebugDataDelegate(class UCanvas* Canvas,
 				ServerSetActorToDebug(NewTarget);
 			}
 
-			GetDebugComponent()->SetActorToDebug(NewTarget);
-			GetDebugComponent()->CollectDataToReplicate(true);
+			// @note patching up OR-9814
+			const UGameplayDebuggingComponent* DummyPointer = GetDebugComponent();
+			if (DebugComponent)
+			{
+				DebugComponent->SetActorToDebug(NewTarget);
+				DebugComponent->CollectDataToReplicate(true);
+			}
 		}
 	}
 
diff --git a/Engine/Source/Editor/AnimGraph/Classes/K2Node_AnimGetter.h b/Engine/Source/Editor/AnimGraph/Classes/K2Node_AnimGetter.h
--- a/Engine/Source/Editor/AnimGraph/Classes/K2Node_AnimGetter.h
+++ b/Engine/Source/Editor/AnimGraph/Classes/K2Node_AnimGetter.h
@@ -78,9 +78,12 @@ public:
 	UPROPERTY()
 	const UAnimBlueprint* SourceAnimBlueprint;
 
+	//~ Using WITH_EDITORONLY_DATA within an Editor module to exclude this FText property from the gather for games
+#if WITH_EDITORONLY_DATA
 	// Cached node title
 	UPROPERTY()
 	FText CachedTitle;
+#endif // WITH_EDITORONLY_DATA
 
 	// List of valid contexts for the node
 	UPROPERTY()
diff --git a/Engine/Source/Editor/BlueprintGraph/Classes/K2Node.h b/Engine/Source/Editor/BlueprintGraph/Classes/K2Node.h
--- a/Engine/Source/Editor/BlueprintGraph/Classes/K2Node.h
+++ b/Engine/Source/Editor/BlueprintGraph/Classes/K2Node.h
@@ -21,8 +21,11 @@ struct FOptionalPinFromProperty
 	UPROPERTY(EditAnywhere, Category=Hi, BlueprintReadOnly)
 	FString PropertyFriendlyName;
 
+	//~ Using WITH_EDITORONLY_DATA within an Editor module to exclude this FText property from the gather for games
+#if WITH_EDITORONLY_DATA
 	UPROPERTY(EditAnywhere, Category=Hi, BlueprintReadOnly)
 	FText PropertyTooltip;
+#endif //~ WITH_EDITORONLY_DATA
 
 	UPROPERTY(EditAnywhere, BlueprintReadWrite, Category=Hi)
 	bool bShowPin;
diff --git a/Engine/Source/Editor/BlueprintGraph/Classes/K2Node_Select.h b/Engine/Source/Editor/BlueprintGraph/Classes/K2Node_Select.h
--- a/Engine/Source/Editor/BlueprintGraph/Classes/K2Node_Select.h
+++ b/Engine/Source/Editor/BlueprintGraph/Classes/K2Node_Select.h
@@ -27,9 +27,12 @@ class UK2Node_Select : public UK2Node, public INodeDependingOnEnumInterface
 	UPROPERTY()
 	TArray<FName> EnumEntries;
 
+	//~ Using WITH_EDITORONLY_DATA within an Editor module to exclude this FText property from the gather for games
+#if WITH_EDITORONLY_DATA
 	/** List of the current entries in the enum (Pin Friendly Names) */
 	UPROPERTY()
 	TArray<FText> EnumEntryFriendlyNames;
+#endif //~ WITH_EDITORONLY_DATA
 
 	/** Whether we need to reconstruct the node after the pins have changed */
 	UPROPERTY(Transient)
diff --git a/Engine/Source/Editor/GraphEditor/Private/SCommentBubble.cpp b/Engine/Source/Editor/GraphEditor/Private/SCommentBubble.cpp
--- a/Engine/Source/Editor/GraphEditor/Private/SCommentBubble.cpp
+++ b/Engine/Source/Editor/GraphEditor/Private/SCommentBubble.cpp
@@ -403,7 +403,7 @@ void SCommentBubble::OnCommentTextCommitted( const FText& NewText, ETextCommit::
 
 EVisibility SCommentBubble::GetToggleButtonVisibility() const
 {
-	EVisibility Visibility = EVisibility::Collapsed;
+	EVisibility Visibility = EVisibility::Hidden;
 
 	if( OpacityValue > 0.f && !GraphNode->bCommentBubbleVisible )
 	{
diff --git a/Engine/Source/Editor/Kismet/Private/FindInBlueprintManager.cpp b/Engine/Source/Editor/Kismet/Private/FindInBlueprintManager.cpp
--- a/Engine/Source/Editor/Kismet/Private/FindInBlueprintManager.cpp
+++ b/Engine/Source/Editor/Kismet/Private/FindInBlueprintManager.cpp
@@ -1271,7 +1271,16 @@ void FFindInBlueprintSearchManager::OnAssetAdded(const FAssetData& InAssetData)
 			}
 			else if(const FString* FiBVersionedSearchData = InAssetData.TagsAndValues.Find("FiBData"))
 			{
-				ExtractUnloadedFiBData(InAssetData, *FiBVersionedSearchData, true);
+				if (FiBVersionedSearchData->Len() == 0)
+				{
+					// agrant TODO: Can we patch this up? Is it dangerous not to?
+					UE_LOG(LogBlueprint, Warning, TEXT("Blueprint %s had empty search data!"), *InAssetData.ObjectPath.ToString());
+					UncachedBlueprints.Add(InAssetData.ObjectPath);
+				}
+				else
+				{
+					ExtractUnloadedFiBData(InAssetData, *FiBVersionedSearchData, true);
+				}
 			}
 			else
 			{
@@ -1319,6 +1328,7 @@ void FFindInBlueprintSearchManager::ExtractUnloadedFiBData(const FAssetData& InA
 	// Deserialize the version if available
 	if (bIsVersioned)
 	{
+		checkf(NewSearchData.Value.Len(), TEXT("Versioned search data was zero length!"));
 		FBufferReader ReaderStream((void*)*NewSearchData.Value, NewSearchData.Value.Len() * sizeof(TCHAR), false);
 		NewSearchData.Version = FiBSerializationHelpers::Deserialize<int32>(ReaderStream);
 	}
diff --git a/Engine/Source/Editor/MovieSceneTools/Private/Sections/ColorPropertySection.cpp b/Engine/Source/Editor/MovieSceneTools/Private/Sections/ColorPropertySection.cpp
--- a/Engine/Source/Editor/MovieSceneTools/Private/Sections/ColorPropertySection.cpp
+++ b/Engine/Source/Editor/MovieSceneTools/Private/Sections/ColorPropertySection.cpp
@@ -163,7 +163,7 @@ FLinearColor FColorPropertySection::FindSlateColor(const FName& ColorName) const
 				continue;
 			}
 
-			UObject* RuntimeObject = FocusedSequence->FindObject(Binding.GetObjectGuid());
+			UObject* RuntimeObject = Sequencer->GetFocusedMovieSceneSequenceInstance()->FindObject(Binding.GetObjectGuid(), *Sequencer);
 
 			if (RuntimeObject == nullptr)
 			{
diff --git a/Engine/Source/Editor/MovieSceneTools/Private/Sections/ShotSequencerSection.cpp b/Engine/Source/Editor/MovieSceneTools/Private/Sections/ShotSequencerSection.cpp
--- a/Engine/Source/Editor/MovieSceneTools/Private/Sections/ShotSequencerSection.cpp
+++ b/Engine/Source/Editor/MovieSceneTools/Private/Sections/ShotSequencerSection.cpp
@@ -166,7 +166,7 @@ bool FShotSequencerSection::AreSectionsConnected() const
 }
 
 
-void FShotSequencerSection::BuildSectionContextMenu(FMenuBuilder& MenuBuilder)
+void FShotSequencerSection::BuildSectionContextMenu(FMenuBuilder& MenuBuilder, const FGuid& ObjectBinding)
 {
 	UWorld* World = GEditor->GetEditorWorldContext().World();
 
diff --git a/Engine/Source/Editor/MovieSceneTools/Private/Sections/ShotSequencerSection.h b/Engine/Source/Editor/MovieSceneTools/Private/Sections/ShotSequencerSection.h
--- a/Engine/Source/Editor/MovieSceneTools/Private/Sections/ShotSequencerSection.h
+++ b/Engine/Source/Editor/MovieSceneTools/Private/Sections/ShotSequencerSection.h
@@ -54,7 +54,7 @@ public:
 
 	virtual bool ShouldDrawKeyAreaBackground() const override;
 	virtual bool AreSectionsConnected() const override;
-	virtual void BuildSectionContextMenu(FMenuBuilder& MenuBuilder) override;
+	virtual void BuildSectionContextMenu(FMenuBuilder& MenuBuilder, const FGuid& ObjectBinding) override;
 	virtual void GenerateSectionLayout(ISectionLayoutBuilder& LayoutBuilder) const override { }
 	virtual TSharedRef<SWidget> GenerateSectionWidget() override;
 	virtual FText GetDisplayName() const override;
diff --git a/Engine/Source/Editor/MovieSceneTools/Private/TrackEditors/ActorPickerTrackEditor.cpp b/Engine/Source/Editor/MovieSceneTools/Private/TrackEditors/ActorPickerTrackEditor.cpp
--- a/Engine/Source/Editor/MovieSceneTools/Private/TrackEditors/ActorPickerTrackEditor.cpp
+++ b/Engine/Source/Editor/MovieSceneTools/Private/TrackEditors/ActorPickerTrackEditor.cpp
@@ -24,7 +24,7 @@ void FActorPickerTrackEditor::PickActorInteractive(FGuid ObjectBinding, UMovieSc
 
 		ActorPickerMode.BeginActorPickingMode(
 			FOnGetAllowedClasses(), 
-			FOnShouldFilterActor::CreateSP(this, &FActorPickerTrackEditor::IsActorPickable), 
+			FOnShouldFilterActor::CreateSP(this, &FActorPickerTrackEditor::IsActorPickable, ObjectBinding, Section), 
 			FOnActorSelected::CreateSP(this, &FActorPickerTrackEditor::ActorPicked, ObjectBinding, Section)
 			);
 	}
@@ -55,7 +55,7 @@ void FActorPickerTrackEditor::ShowActorSubMenu(FMenuBuilder& MenuBuilder, FGuid
 		InitOptions.ColumnMap.Add(FBuiltInColumnTypes::Label(), FColumnInfo(EColumnVisibility::Visible, 0));
 
 		// Only display Actors that we can attach too
-		InitOptions.Filters->AddFilterPredicate( SceneOutliner::FActorFilterPredicate::CreateSP(this, &FActorPickerTrackEditor::IsActorPickable) );
+		InitOptions.Filters->AddFilterPredicate( SceneOutliner::FActorFilterPredicate::CreateSP(this, &FActorPickerTrackEditor::IsActorPickable, ObjectBinding, Section) );
 	}		
 
 	// Actor selector to allow the user to choose a parent actor
diff --git a/Engine/Source/Editor/MovieSceneTools/Private/TrackEditors/ActorPickerTrackEditor.h b/Engine/Source/Editor/MovieSceneTools/Private/TrackEditors/ActorPickerTrackEditor.h
--- a/Engine/Source/Editor/MovieSceneTools/Private/TrackEditors/ActorPickerTrackEditor.h
+++ b/Engine/Source/Editor/MovieSceneTools/Private/TrackEditors/ActorPickerTrackEditor.h
@@ -12,7 +12,7 @@ public:
 	FActorPickerTrackEditor(TSharedRef<ISequencer> InSequencer);
 
 	/** Is this actor pickable? */
-	virtual bool IsActorPickable( const AActor* const ParentActor ) { return false; }
+	virtual bool IsActorPickable( const AActor* const ParentActor, FGuid ObjectBinding, UMovieSceneSection* InSection) { return false; }
 
 	/** Actor socked was picked */
 	virtual void ActorSocketPicked(const FName SocketName, AActor* ParentActor, FGuid ObjectBinding, UMovieSceneSection* Section) {}
diff --git a/Engine/Source/Editor/MovieSceneTools/Private/TrackEditors/AttachTrackEditor.cpp b/Engine/Source/Editor/MovieSceneTools/Private/TrackEditors/AttachTrackEditor.cpp
--- a/Engine/Source/Editor/MovieSceneTools/Private/TrackEditors/AttachTrackEditor.cpp
+++ b/Engine/Source/Editor/MovieSceneTools/Private/TrackEditors/AttachTrackEditor.cpp
@@ -64,13 +64,11 @@ public:
 		return LayerId;
 	}
 	
-	virtual void BuildSectionContextMenu(FMenuBuilder& MenuBuilder) override
+	virtual void BuildSectionContextMenu(FMenuBuilder& MenuBuilder, const FGuid& ObjectBinding) override
 	{
-		FGuid DummyObjectBinding;
-
 		MenuBuilder.AddSubMenu(
 			LOCTEXT("SetAttach", "Attach"), LOCTEXT("SetAttachTooltip", "Set attach"),
-			FNewMenuDelegate::CreateRaw(AttachTrackEditor, &FActorPickerTrackEditor::ShowActorSubMenu, DummyObjectBinding, &Section));
+			FNewMenuDelegate::CreateRaw(AttachTrackEditor, &FActorPickerTrackEditor::ShowActorSubMenu, ObjectBinding, &Section));
 	}
 
 private:
@@ -138,8 +136,16 @@ void F3DAttachTrackEditor::BuildObjectBindingTrackMenu(FMenuBuilder& MenuBuilder
 }
 
 
-bool F3DAttachTrackEditor::IsActorPickable(const AActor* const ParentActor)
+bool F3DAttachTrackEditor::IsActorPickable(const AActor* const ParentActor, FGuid ObjectBinding, UMovieSceneSection* InSection)
 {
+	// Can't pick the object that this track binds
+	TArray<UObject*> OutObjects;
+	GetSequencer()->GetRuntimeObjects( GetSequencer()->GetFocusedMovieSceneSequenceInstance(), ObjectBinding, OutObjects);
+	if (OutObjects.Contains(ParentActor))
+	{
+		return false;
+	}
+
 	if (ParentActor->IsListedInSceneOutliner() &&
 		!FActorEditorUtils::IsABuilderBrush(ParentActor) &&
 		!ParentActor->IsA( AWorldSettings::StaticClass() ) &&
diff --git a/Engine/Source/Editor/MovieSceneTools/Private/TrackEditors/AttachTrackEditor.h b/Engine/Source/Editor/MovieSceneTools/Private/TrackEditors/AttachTrackEditor.h
--- a/Engine/Source/Editor/MovieSceneTools/Private/TrackEditors/AttachTrackEditor.h
+++ b/Engine/Source/Editor/MovieSceneTools/Private/TrackEditors/AttachTrackEditor.h
@@ -43,7 +43,7 @@ public:
 	virtual bool SupportsType( TSubclassOf<UMovieSceneTrack> Type ) const override;
 
 	// FTrackEditorActorPicker
-	virtual bool IsActorPickable( const AActor* const ParentActor ) override;
+	virtual bool IsActorPickable( const AActor* const ParentActor, FGuid ObjectBinding, UMovieSceneSection* InSection ) override;
 	virtual void ActorSocketPicked(const FName SocketName, AActor* ParentActor, FGuid ObjectBinding, UMovieSceneSection* Section) override;
 
 private:
diff --git a/Engine/Source/Editor/MovieSceneTools/Private/TrackEditors/MaterialTrackEditor.cpp b/Engine/Source/Editor/MovieSceneTools/Private/TrackEditors/MaterialTrackEditor.cpp
--- a/Engine/Source/Editor/MovieSceneTools/Private/TrackEditors/MaterialTrackEditor.cpp
+++ b/Engine/Source/Editor/MovieSceneTools/Private/TrackEditors/MaterialTrackEditor.cpp
@@ -198,7 +198,7 @@ bool FComponentMaterialTrackEditor::SupportsType( TSubclassOf<UMovieSceneTrack>
 
 UMaterialInterface* FComponentMaterialTrackEditor::GetMaterialInterfaceForTrack( FGuid ObjectBinding, UMovieSceneMaterialTrack* MaterialTrack )
 {
-	UObject* ComponentObject = GetSequencer()->GetFocusedMovieSceneSequence()->FindObject( ObjectBinding );
+	UObject* ComponentObject = GetSequencer()->GetFocusedMovieSceneSequenceInstance()->FindObject( ObjectBinding, *GetSequencer() );
 	UPrimitiveComponent* Component = Cast<UPrimitiveComponent>( ComponentObject );
 	UMovieSceneComponentMaterialTrack* ComponentMaterialTrack = Cast<UMovieSceneComponentMaterialTrack>( MaterialTrack );
 	if ( Component != nullptr && ComponentMaterialTrack != nullptr )
diff --git a/Engine/Source/Editor/MovieSceneTools/Private/TrackEditors/ParticleParameterTrackEditor.cpp b/Engine/Source/Editor/MovieSceneTools/Private/TrackEditors/ParticleParameterTrackEditor.cpp
--- a/Engine/Source/Editor/MovieSceneTools/Private/TrackEditors/ParticleParameterTrackEditor.cpp
+++ b/Engine/Source/Editor/MovieSceneTools/Private/TrackEditors/ParticleParameterTrackEditor.cpp
@@ -94,7 +94,7 @@ TSharedRef<SWidget> FParticleParameterTrackEditor::OnGetAddParameterMenuContent(
 {
 	FMenuBuilder AddParameterMenuBuilder( true, nullptr );
 
-	AEmitter* Emitter = Cast<AEmitter>( GetSequencer()->GetFocusedMovieSceneSequence()->FindObject( ObjectBinding ) );
+	AEmitter* Emitter = Cast<AEmitter>( GetSequencer()->GetFocusedMovieSceneSequenceInstance()->FindObject( ObjectBinding, *GetSequencer() ) );
 	if ( Emitter != nullptr )
 	{
 		UParticleSystemComponent* ParticleSystemComponent = Emitter->GetParticleSystemComponent();
@@ -161,7 +161,7 @@ void FParticleParameterTrackEditor::AddScalarParameter( FGuid ObjectBinding, UMo
 	UMovieSceneSequence* MovieSceneSequence = GetMovieSceneSequence();
 	float KeyTime = GetTimeForKey( MovieSceneSequence );
 
-	UObject* Object = GetSequencer()->GetFocusedMovieSceneSequence()->FindObject( ObjectBinding );
+	UObject* Object = GetSequencer()->GetFocusedMovieSceneSequenceInstance()->FindObject( ObjectBinding, *GetSequencer() );
 	AEmitter* Emitter = Cast<AEmitter>(Object);
 	if ( Emitter != nullptr )
 	{
@@ -185,7 +185,7 @@ void FParticleParameterTrackEditor::AddVectorParameter( FGuid ObjectBinding, UMo
 	UMovieSceneSequence* MovieSceneSequence = GetMovieSceneSequence();
 	float KeyTime = GetTimeForKey( MovieSceneSequence );
 
-	UObject* Object = GetSequencer()->GetFocusedMovieSceneSequence()->FindObject( ObjectBinding );
+	UObject* Object = GetSequencer()->GetFocusedMovieSceneSequenceInstance()->FindObject( ObjectBinding, *GetSequencer() );
 	AEmitter* Emitter = Cast<AEmitter>( Object );
 	if ( Emitter != nullptr )
 	{
@@ -209,7 +209,7 @@ void FParticleParameterTrackEditor::AddColorParameter( FGuid ObjectBinding, UMov
 	UMovieSceneSequence* MovieSceneSequence = GetMovieSceneSequence();
 	float KeyTime = GetTimeForKey( MovieSceneSequence );
 
-	UObject* Object = GetSequencer()->GetFocusedMovieSceneSequence()->FindObject( ObjectBinding );
+	UObject* Object = GetSequencer()->GetFocusedMovieSceneSequenceInstance()->FindObject( ObjectBinding, *GetSequencer() );
 	AEmitter* Emitter = Cast<AEmitter>( Object );
 	if ( Emitter != nullptr )
 	{
diff --git a/Engine/Source/Editor/MovieSceneTools/Private/TrackEditors/ParticleTrackEditor.cpp b/Engine/Source/Editor/MovieSceneTools/Private/TrackEditors/ParticleTrackEditor.cpp
--- a/Engine/Source/Editor/MovieSceneTools/Private/TrackEditors/ParticleTrackEditor.cpp
+++ b/Engine/Source/Editor/MovieSceneTools/Private/TrackEditors/ParticleTrackEditor.cpp
@@ -102,7 +102,7 @@ int32 FParticleSection::OnPaintSection( const FGeometry& AllottedGeometry, const
 
 			if ( ObjectHandle.IsValid() )
 			{
-				AEmitter* ParticleSystemActor = Cast<AEmitter>( OwningSequencer->GetFocusedMovieSceneSequence()->FindObject( ObjectHandle ) );
+				AEmitter* ParticleSystemActor = Cast<AEmitter>( OwningSequencer->GetFocusedMovieSceneSequenceInstance()->FindObject( ObjectHandle, *OwningSequencer ) );
 				if ( ParticleSystemActor != nullptr )
 				{
 					UParticleSystemComponent* ParticleSystemComponent = ParticleSystemActor->GetParticleSystemComponent();
diff --git a/Engine/Source/Editor/MovieSceneTools/Private/TrackEditors/PathTrackEditor.cpp b/Engine/Source/Editor/MovieSceneTools/Private/TrackEditors/PathTrackEditor.cpp
--- a/Engine/Source/Editor/MovieSceneTools/Private/TrackEditors/PathTrackEditor.cpp
+++ b/Engine/Source/Editor/MovieSceneTools/Private/TrackEditors/PathTrackEditor.cpp
@@ -65,13 +65,11 @@ public:
 		return LayerId;
 	}
 	
-	virtual void BuildSectionContextMenu(FMenuBuilder& MenuBuilder) override
+	virtual void BuildSectionContextMenu(FMenuBuilder& MenuBuilder, const FGuid& ObjectBinding) override
 	{
-		FGuid DummyObjectBinding;
-
 		MenuBuilder.AddSubMenu(
 			LOCTEXT("SetPath", "Path"), LOCTEXT("SetPathTooltip", "Set path"),
-			FNewMenuDelegate::CreateRaw(PathTrackEditor, &FActorPickerTrackEditor::ShowActorSubMenu, DummyObjectBinding, &Section));
+			FNewMenuDelegate::CreateRaw(PathTrackEditor, &FActorPickerTrackEditor::ShowActorSubMenu, ObjectBinding, &Section));
 	}
 
 private:
@@ -134,8 +132,33 @@ void F3DPathTrackEditor::BuildObjectBindingTrackMenu(FMenuBuilder& MenuBuilder,
 	}
 }
 
-bool F3DPathTrackEditor::IsActorPickable(const AActor* const ParentActor)
+bool F3DPathTrackEditor::IsActorPickable(const AActor* const ParentActor, FGuid ObjectBinding, UMovieSceneSection* InSection)
 {
+	// Can't pick the object that this track binds
+	TArray<UObject*> OutObjects;
+	GetSequencer()->GetRuntimeObjects( GetSequencer()->GetFocusedMovieSceneSequenceInstance(), ObjectBinding, OutObjects);
+	if (OutObjects.Contains(ParentActor))
+	{
+		return false;
+	}
+
+	// Can't pick the object that this track attaches to
+	UMovieScene3DPathSection* PathSection = Cast<UMovieScene3DPathSection>(InSection);
+	if (PathSection != nullptr)
+	{
+		FGuid ConstraintId = PathSection->GetConstraintId();
+
+		if (ConstraintId.IsValid())
+		{
+			TArray<UObject*> ConstraintObjects;
+			GetSequencer()->GetRuntimeObjects( GetSequencer()->GetFocusedMovieSceneSequenceInstance(), ConstraintId, ConstraintObjects);
+			if (ConstraintObjects.Contains(ParentActor))
+			{
+				return false;
+			}
+		}
+	}
+
 	if (ParentActor->IsListedInSceneOutliner() &&
 		!FActorEditorUtils::IsABuilderBrush(ParentActor) &&
 		!ParentActor->IsA( AWorldSettings::StaticClass() ) &&
diff --git a/Engine/Source/Editor/MovieSceneTools/Private/TrackEditors/PathTrackEditor.h b/Engine/Source/Editor/MovieSceneTools/Private/TrackEditors/PathTrackEditor.h
--- a/Engine/Source/Editor/MovieSceneTools/Private/TrackEditors/PathTrackEditor.h
+++ b/Engine/Source/Editor/MovieSceneTools/Private/TrackEditors/PathTrackEditor.h
@@ -43,7 +43,7 @@ public:
 	virtual bool SupportsType( TSubclassOf<UMovieSceneTrack> Type ) const override;
 
 	// FTrackEditorActorPicker
-	virtual bool IsActorPickable( const AActor* const ParentActor ) override;
+	virtual bool IsActorPickable( const AActor* const ParentActor, FGuid ObjectBinding, UMovieSceneSection* InSection ) override;
 	virtual void ActorSocketPicked(const FName SocketName, AActor* ParentActor, FGuid ObjectBinding, UMovieSceneSection* Section) override;
 
 private:
diff --git a/Engine/Source/Editor/MovieSceneTools/Private/TrackEditors/PropertyTrackEditors/BytePropertyTrackEditor.cpp b/Engine/Source/Editor/MovieSceneTools/Private/TrackEditors/PropertyTrackEditors/BytePropertyTrackEditor.cpp
--- a/Engine/Source/Editor/MovieSceneTools/Private/TrackEditors/PropertyTrackEditors/BytePropertyTrackEditor.cpp
+++ b/Engine/Source/Editor/MovieSceneTools/Private/TrackEditors/PropertyTrackEditors/BytePropertyTrackEditor.cpp
@@ -22,7 +22,7 @@ TSharedRef<FPropertySection> FBytePropertyTrackEditor::MakePropertySectionInterf
 UEnum* GetEnumForByteTrack(TSharedPtr<ISequencer> Sequencer, const FGuid& OwnerObjectHandle, FName PropertyName, UMovieSceneByteTrack* ByteTrack)
 {
 	
-	UObject* RuntimeObject = Sequencer->GetFocusedMovieSceneSequence()->FindObject(OwnerObjectHandle);
+	UObject* RuntimeObject = Sequencer->GetFocusedMovieSceneSequenceInstance()->FindObject(OwnerObjectHandle, *Sequencer);
 	TSet<UEnum*> PropertyEnums;
 
 	if (RuntimeObject != nullptr)
diff --git a/Engine/Source/Editor/MovieSceneTools/Private/TrackEditors/PropertyTrackEditors/VectorPropertyTrackEditor.cpp b/Engine/Source/Editor/MovieSceneTools/Private/TrackEditors/PropertyTrackEditors/VectorPropertyTrackEditor.cpp
--- a/Engine/Source/Editor/MovieSceneTools/Private/TrackEditors/PropertyTrackEditors/VectorPropertyTrackEditor.cpp
+++ b/Engine/Source/Editor/MovieSceneTools/Private/TrackEditors/PropertyTrackEditors/VectorPropertyTrackEditor.cpp
@@ -58,22 +58,16 @@ void FVectorPropertyTrackEditor::GenerateKeysFromPropertyChanged( const FPropert
 	}
 
 	FName ChannelName = PropertyChangedParams.StructPropertyNameToKey;
-
-	if ( ChannelName == NAME_None || ChannelName == XName )
-	{
-		GeneratedKeys.Add( FVectorKey( EKeyVectorChannel::X, VectorValues.X ) );
-	}
-	if ( ChannelName == NAME_None || ChannelName == YName )
-	{
-		GeneratedKeys.Add( FVectorKey( EKeyVectorChannel::Y, VectorValues.Y ) );
-	}
-	if ( (ChannelName == NAME_None || ChannelName == ZName) && Channels >= 3 )
+		
+	GeneratedKeys.Add( FVectorKey( EKeyVectorChannel::X, VectorValues.X, ChannelName == NAME_None || ChannelName == XName ? EKeyVectorValueType::Key : EKeyVectorValueType::Default ) );
+	GeneratedKeys.Add( FVectorKey( EKeyVectorChannel::Y, VectorValues.Y, ChannelName == NAME_None || ChannelName == YName ? EKeyVectorValueType::Key : EKeyVectorValueType::Default ) );
+	if (Channels >= 3 )
 	{
-		GeneratedKeys.Add( FVectorKey( EKeyVectorChannel::Z, VectorValues.Z ) );
+		GeneratedKeys.Add( FVectorKey( EKeyVectorChannel::Z, VectorValues.Z, ChannelName == NAME_None || ChannelName == ZName ? EKeyVectorValueType::Key : EKeyVectorValueType::Default ) );
 	}
-	if ( (ChannelName == NAME_None || ChannelName == WName) && Channels >= 4 )
+	if (Channels >= 4 )
 	{
-		GeneratedKeys.Add( FVectorKey( EKeyVectorChannel::W, VectorValues.W ) );
+		GeneratedKeys.Add( FVectorKey( EKeyVectorChannel::W, VectorValues.W, ChannelName == NAME_None || ChannelName == WName ? EKeyVectorValueType::Key : EKeyVectorValueType::Default ) );
 	}
 }
 
@@ -96,3 +90,8 @@ void FVectorPropertyTrackEditor::InitializeNewTrack( UMovieSceneVectorTrack* New
 		NewTrack->SetNumChannelsUsed( 4 );
 	}
 }
+
+bool FVectorPropertyTrackEditor::ShouldAddKey(UMovieSceneVectorTrack* InTrack, FVectorKey InKey, FKeyParams InKeyParams) const
+{
+	return FPropertyTrackEditor::ShouldAddKey(InTrack, InKey, InKeyParams) && InKey.ValueType == EKeyVectorValueType::Key;
+}
diff --git a/Engine/Source/Editor/MovieSceneTools/Private/TrackEditors/PropertyTrackEditors/VectorPropertyTrackEditor.h b/Engine/Source/Editor/MovieSceneTools/Private/TrackEditors/PropertyTrackEditors/VectorPropertyTrackEditor.h
--- a/Engine/Source/Editor/MovieSceneTools/Private/TrackEditors/PropertyTrackEditors/VectorPropertyTrackEditor.h
+++ b/Engine/Source/Editor/MovieSceneTools/Private/TrackEditors/PropertyTrackEditors/VectorPropertyTrackEditor.h
@@ -38,6 +38,9 @@ protected:
 	virtual void GenerateKeysFromPropertyChanged( const FPropertyChangedParams& PropertyChangedParams, TArray<FVectorKey>& GeneratedKeys ) override;
 	virtual void InitializeNewTrack( UMovieSceneVectorTrack* NewTrack, FPropertyChangedParams PropertyChangedParams ) override;
 
+	// FKeyframeTrackEditor interface
+	virtual bool ShouldAddKey(UMovieSceneVectorTrack* InTrack, FVectorKey InKey, FKeyParams InKeyParams) const override;
+
 private:
 	static FName XName;
 	static FName YName;
diff --git a/Engine/Source/Editor/MovieSceneTools/Private/TrackEditors/SpawnTrackEditor.cpp b/Engine/Source/Editor/MovieSceneTools/Private/TrackEditors/SpawnTrackEditor.cpp
--- a/Engine/Source/Editor/MovieSceneTools/Private/TrackEditors/SpawnTrackEditor.cpp
+++ b/Engine/Source/Editor/MovieSceneTools/Private/TrackEditors/SpawnTrackEditor.cpp
@@ -26,13 +26,77 @@ UMovieSceneTrack* FSpawnTrackEditor::AddTrack(UMovieScene* FocusedMovieScene, co
 
 	if (auto* SpawnTrack = Cast<UMovieSceneSpawnTrack>(NewTrack))
 	{
-		SpawnTrack->SetObject(ObjectHandle);
+		SpawnTrack->SetObjectId(ObjectHandle);
 		SpawnTrack->AddSection(*SpawnTrack->CreateNewSection());
 	}
 
 	return NewTrack;
 }
 
+void FSpawnTrackEditor::BuildTrackContextMenu(FMenuBuilder& MenuBuilder, UMovieSceneTrack* Track)
+{
+	UMovieSceneSpawnTrack* SpawnTrack = Cast<UMovieSceneSpawnTrack>(Track);
+	if (!SpawnTrack)
+	{
+		return;
+	}
+
+	FMovieSceneSpawnable* Spawnable = GetSequencer()->GetFocusedMovieSceneSequence()->GetMovieScene()->FindSpawnable(SpawnTrack->GetObjectId());
+	if (!Spawnable)
+	{
+		return;
+	}
+
+	MenuBuilder.AddSubMenu(
+		LOCTEXT("OwnerLabel", "Spawned Object Owner"),
+		LOCTEXT("OwnerTooltip", "Specifies how the object spawned in this track is to be owned"),
+		FNewMenuDelegate::CreateSP(this, &FSpawnTrackEditor::AddSpawnOwnershipMenu, SpawnTrack)
+	);
+}
+
+void FSpawnTrackEditor::AddSpawnOwnershipMenu(FMenuBuilder& MenuBuilder, UMovieSceneSpawnTrack* Track)
+{
+	FMovieSceneSpawnable* Spawnable = GetSequencer()->GetFocusedMovieSceneSequence()->GetMovieScene()->FindSpawnable(Track->GetObjectId());
+
+	MenuBuilder.AddMenuEntry(
+		LOCTEXT("ThisSequence_Label", "This Sequence"),
+		LOCTEXT("ThisSequence_Tooltip", "Indicates that this sequence will own the spawned object. The object will be destroyed at the end of the sequence."),
+		FSlateIcon(),
+		FUIAction(
+			FExecuteAction::CreateLambda([=]{ Spawnable->SetSpawnOwnership(ESpawnOwnership::InnerSequence); }),
+			FCanExecuteAction(),
+			FIsActionChecked::CreateLambda([=]{ return Spawnable->GetSpawnOwnership() == ESpawnOwnership::InnerSequence; })
+		),
+		NAME_None,
+		EUserInterfaceActionType::ToggleButton
+	);
+
+	MenuBuilder.AddMenuEntry(
+		LOCTEXT("MasterSequence_Label", "Master Sequence"),
+		LOCTEXT("MasterSequence_Tooltip", "Indicates that the outermost sequence will own the spawned object. The object will be destroyed when the outermost sequence stops playing."),
+		FSlateIcon(),
+		FUIAction(
+			FExecuteAction::CreateLambda([=]{ Spawnable->SetSpawnOwnership(ESpawnOwnership::MasterSequence); }),
+			FCanExecuteAction(),
+			FIsActionChecked::CreateLambda([=]{ return Spawnable->GetSpawnOwnership() == ESpawnOwnership::MasterSequence; })
+		),
+		NAME_None,
+		EUserInterfaceActionType::ToggleButton
+	);
+
+	MenuBuilder.AddMenuEntry(
+		LOCTEXT("External_Label", "External"),
+		LOCTEXT("External_Tooltip", "Indicates this object's lifetime is managed externally once spawned. It will not be destroyed by sequencer."),
+		FSlateIcon(),
+		FUIAction(
+			FExecuteAction::CreateLambda([=]{ Spawnable->SetSpawnOwnership(ESpawnOwnership::External); }),
+			FCanExecuteAction(),
+			FIsActionChecked::CreateLambda([=]{ return Spawnable->GetSpawnOwnership() == ESpawnOwnership::External; })
+		),
+		NAME_None,
+		EUserInterfaceActionType::ToggleButton
+	);
+}
 
 void FSpawnTrackEditor::BuildObjectBindingTrackMenu(FMenuBuilder& MenuBuilder, const FGuid& ObjectBinding, const UClass* ObjectClass)
 {
diff --git a/Engine/Source/Editor/MovieSceneTools/Private/TrackEditors/SpawnTrackEditor.h b/Engine/Source/Editor/MovieSceneTools/Private/TrackEditors/SpawnTrackEditor.h
--- a/Engine/Source/Editor/MovieSceneTools/Private/TrackEditors/SpawnTrackEditor.h
+++ b/Engine/Source/Editor/MovieSceneTools/Private/TrackEditors/SpawnTrackEditor.h
@@ -7,7 +7,7 @@
 
 class ISequencer;
 class UMovieSceneTrack;
-
+class UMovieSceneSpawnTrack;
 
 /**
  * A property track editor for controlling the lifetime of a sapwnable object
@@ -43,10 +43,14 @@ public:
 	virtual TSharedPtr<SWidget> BuildOutlinerEditWidget(const FGuid& ObjectBinding, UMovieSceneTrack* Track) override  { return TSharedPtr<SWidget>(); }
 	virtual bool HandleAssetAdded(UObject* Asset, const FGuid& TargetObjectGuid) override { return false; }
 	virtual bool SupportsType(TSubclassOf<UMovieSceneTrack> Type) const override;
+	virtual void BuildTrackContextMenu( FMenuBuilder& MenuBuilder, UMovieSceneTrack* Track ) override;
 
 private:
 
 	/** Callback for executing the "Add Spawn Track" menu entry. */
 	void HandleAddSpawnTrackMenuEntryExecute(FGuid ObjectBinding);
 	bool CanAddSpawnTrack(FGuid ObjectBinding) const;
+
+	/** Populate the sub menu for setting the ownership of the spawned object */
+	void AddSpawnOwnershipMenu(FMenuBuilder& MenuBuilder, UMovieSceneSpawnTrack* Track);
 };
diff --git a/Engine/Source/Editor/MovieSceneTools/Private/TrackEditors/TransformTrackEditor.cpp b/Engine/Source/Editor/MovieSceneTools/Private/TrackEditors/TransformTrackEditor.cpp
--- a/Engine/Source/Editor/MovieSceneTools/Private/TrackEditors/TransformTrackEditor.cpp
+++ b/Engine/Source/Editor/MovieSceneTools/Private/TrackEditors/TransformTrackEditor.cpp
@@ -643,38 +643,33 @@ void F3DTransformTrackEditor::GetActorAndSceneComponentFromObject( UObject* Obje
 	}
 }
 
-void GetKeysForVector( bool LastVectorIsValid, const FVector& LastVector, const FVector& CurrentVector, EKey3DTransformChannel::Type KeyChannel, bool bUnwindRotation, TArray<FTransformKey>& OutKeys )
+void GetKeysForVector( bool LastVectorIsValid, const FVector& LastVector, const FVector& CurrentVector, EKey3DTransformChannel::Type KeyChannel, EKey3DTransformChannel::ValueType KeyChannelValueType, bool bUnwindRotation, TArray<FTransformKey>& OutKeys )
 {
 	if ( LastVectorIsValid == false || FMath::IsNearlyEqual( LastVector.X, CurrentVector.X ) == false )
 	{
-		OutKeys.Add( FTransformKey( KeyChannel, EAxis::X, CurrentVector.X, bUnwindRotation ) );
+		OutKeys.Add( FTransformKey( KeyChannel, KeyChannelValueType, EAxis::X, CurrentVector.X, bUnwindRotation ) );
 	}
 	if ( LastVectorIsValid == false || FMath::IsNearlyEqual( LastVector.Y, CurrentVector.Y ) == false )
 	{
-		OutKeys.Add( FTransformKey( KeyChannel, EAxis::Y, CurrentVector.Y, bUnwindRotation ) );
+		OutKeys.Add( FTransformKey( KeyChannel, KeyChannelValueType, EAxis::Y, CurrentVector.Y, bUnwindRotation ) );
 	}
 	if ( LastVectorIsValid == false || FMath::IsNearlyEqual( LastVector.Z, CurrentVector.Z ) == false )
 	{
-		OutKeys.Add( FTransformKey( KeyChannel, EAxis::Z, CurrentVector.Z, bUnwindRotation ) );
+		OutKeys.Add( FTransformKey( KeyChannel, KeyChannelValueType, EAxis::Z, CurrentVector.Z, bUnwindRotation ) );
 	}
 }
 
 void F3DTransformTrackEditor::GetChangedTransformKeys( const FTransformData& LastTransform, const FTransformData& CurrentTransform, EKey3DTransformChannel::Type ChannelsToKey, bool bUnwindRotation, TArray<FTransformKey>& OutKeys )
 {
-	if ( ChannelsToKey & EKey3DTransformChannel::Translation )
-	{
-		GetKeysForVector(LastTransform.IsValid(), LastTransform.Translation, CurrentTransform.Translation, EKey3DTransformChannel::Translation, bUnwindRotation, OutKeys);
-	}
-	if ( ChannelsToKey & EKey3DTransformChannel::Rotation )
-	{
-		GetKeysForVector( LastTransform.IsValid(), LastTransform.Rotation.Euler(), CurrentTransform.Rotation.Euler(), EKey3DTransformChannel::Rotation, bUnwindRotation, OutKeys );
-	}
-	if ( ChannelsToKey & EKey3DTransformChannel::Scale )
-	{
-		GetKeysForVector( LastTransform.IsValid(), LastTransform.Scale, CurrentTransform.Scale, EKey3DTransformChannel::Scale, bUnwindRotation, OutKeys );
-	}
+	GetKeysForVector( LastTransform.IsValid(), LastTransform.Translation, CurrentTransform.Translation, EKey3DTransformChannel::Translation, ChannelsToKey & EKey3DTransformChannel::Translation ? EKey3DTransformChannel::Key : EKey3DTransformChannel::Default, bUnwindRotation, OutKeys);
+	GetKeysForVector( LastTransform.IsValid(), LastTransform.Rotation.Euler(), CurrentTransform.Rotation.Euler(), EKey3DTransformChannel::Rotation, ChannelsToKey & EKey3DTransformChannel::Rotation ? EKey3DTransformChannel::Key : EKey3DTransformChannel::Default, bUnwindRotation, OutKeys );
+	GetKeysForVector( LastTransform.IsValid(), LastTransform.Scale, CurrentTransform.Scale, EKey3DTransformChannel::Scale, ChannelsToKey & EKey3DTransformChannel::Scale ? EKey3DTransformChannel::Key : EKey3DTransformChannel::Default, bUnwindRotation, OutKeys );
 }
 
+bool F3DTransformTrackEditor::ShouldAddKey(UMovieScene3DTransformTrack* InTrack, FTransformKey InKey, FKeyParams InKeyParams) const
+{
+	return FKeyframeTrackEditor::ShouldAddKey(InTrack, InKey, InKeyParams) && InKey.ChannelValueType == EKey3DTransformChannel::Key;
+}
 
 void F3DTransformTrackEditor::AddTransformKeysForHandle( FGuid ObjectHandle, EKey3DTransformChannel::Type ChannelToKey, FKeyParams KeyParams )
 {
diff --git a/Engine/Source/Editor/MovieSceneTools/Private/TrackEditors/TransformTrackEditor.h b/Engine/Source/Editor/MovieSceneTools/Private/TrackEditors/TransformTrackEditor.h
--- a/Engine/Source/Editor/MovieSceneTools/Private/TrackEditors/TransformTrackEditor.h
+++ b/Engine/Source/Editor/MovieSceneTools/Private/TrackEditors/TransformTrackEditor.h
@@ -46,6 +46,9 @@ public:
 	virtual bool SupportsType( TSubclassOf<UMovieSceneTrack> Type ) const override;
 	virtual void BuildTrackContextMenu( FMenuBuilder& MenuBuilder, UMovieSceneTrack* Track ) override;
 
+	// FKeyframeTrackEditor interface
+	virtual bool ShouldAddKey(UMovieScene3DTransformTrack* InTrack, FTransformKey InKey, FKeyParams InKeyParams) const override;
+
 private:
 
 	/** Returns whether or not a transform track can be added for an actor with a specific handle. */
diff --git a/Engine/Source/Editor/MovieSceneTools/Public/KeyframeTrackEditor.h b/Engine/Source/Editor/MovieSceneTools/Public/KeyframeTrackEditor.h
--- a/Engine/Source/Editor/MovieSceneTools/Public/KeyframeTrackEditor.h
+++ b/Engine/Source/Editor/MovieSceneTools/Public/KeyframeTrackEditor.h
@@ -39,6 +39,11 @@ public:
 		return Type == TrackType::StaticClass();
 	}
 
+	// Only add a key if the track already has keys, or if creating keys has been forced.  Otherwise just update the track default.
+	virtual bool ShouldAddKey(TrackType* InTrack, KeyDataType InNewKey, FKeyParams InKeyParams) const
+	{
+		return HasKeys( InTrack, InNewKey ) || InKeyParams.bCreateKeyIfEmpty;
+	}
 
 protected:
 
@@ -132,8 +137,7 @@ private:
 				// Only modify the track if the new key will create new data, or if creating keys has been forced.
 				if ( NewKeyIsNewData( Track, KeyTime, NewKey ) || KeyParams.bCreateKeyIfUnchanged )
 				{
-					// Only add a key if the track already has keys, or if creating keys has been forced.  Otherwise just update the track default.
-					if ( HasKeys( Track, NewKey ) || KeyParams.bCreateKeyIfEmpty )
+					if ( ShouldAddKey(Track, NewKey, KeyParams) )
 					{
 						AddKey( Track, KeyTime, NewKey, InterpolationMode );
 					}
diff --git a/Engine/Source/Editor/NiagaraEditor/Private/NiagaraSequencer.cpp b/Engine/Source/Editor/NiagaraEditor/Private/NiagaraSequencer.cpp
--- a/Engine/Source/Editor/NiagaraEditor/Private/NiagaraSequencer.cpp
+++ b/Engine/Source/Editor/NiagaraEditor/Private/NiagaraSequencer.cpp
@@ -34,12 +34,4 @@ void INiagaraTrackInstance::RefreshInstance(const TArray<UObject*>& RuntimeObjec
 }
 
 
-#if WITH_EDITORONLY_DATA
-FText UEmitterMovieSceneTrack::GetDisplayName() const
-{
-	return LOCTEXT("TrackName", "Emitter");
-}
-#endif
-
-
 #undef LOCTEXT_NAMESPACE
diff --git a/Engine/Source/Editor/NiagaraEditor/Private/NiagaraSequencer.h b/Engine/Source/Editor/NiagaraEditor/Private/NiagaraSequencer.h
--- a/Engine/Source/Editor/NiagaraEditor/Private/NiagaraSequencer.h
+++ b/Engine/Source/Editor/NiagaraEditor/Private/NiagaraSequencer.h
@@ -5,7 +5,7 @@
 #include "Editor/Sequencer/Public/MovieSceneTrackEditor.h"
 #include "Editor/Sequencer/Public/ISequencerSection.h"
 #include "Runtime/MovieScene/Public/MovieScene.h"
-#include "Runtime/MovieScene/Public/MovieSceneTrack.h"
+#include "Runtime/MovieScene/Public/MovieSceneNameableTrack.h"
 #include "Runtime/MovieScene/Public/IMovieSceneTrackInstance.h"
 #include "Runtime/MovieScene/Public/MovieSceneSection.h"
 
@@ -141,7 +141,7 @@ private:
 */
 UCLASS(MinimalAPI)
 class UEmitterMovieSceneTrack
-	: public UMovieSceneTrack
+	: public UMovieSceneNameableTrack
 {
 	GENERATED_UCLASS_BODY()
 public:
@@ -208,10 +208,6 @@ public:
 		return TRange<float>(0, FLT_MAX);
 	}
 
-#if WITH_EDITORONLY_DATA
-	virtual FText GetDisplayName() const override;
-#endif
-
 private:
 
 	TSharedPtr<FNiagaraSimulation> Emitter;
@@ -219,7 +215,6 @@ private:
 };
 
 
-
 /**
 *	Track editor for Niagara emitter tracks
 */
diff --git a/Engine/Source/Editor/Sequencer/Private/DisplayNodes/SequencerDisplayNode.cpp b/Engine/Source/Editor/Sequencer/Private/DisplayNodes/SequencerDisplayNode.cpp
--- a/Engine/Source/Editor/Sequencer/Private/DisplayNodes/SequencerDisplayNode.cpp
+++ b/Engine/Source/Editor/Sequencer/Private/DisplayNodes/SequencerDisplayNode.cpp
@@ -416,48 +416,52 @@ void FSequencerDisplayNode::BuildContextMenu(FMenuBuilder& MenuBuilder)
 {
 	TSharedRef<FSequencerDisplayNode> ThisNode = SharedThis(this);
 
-	MenuBuilder.AddMenuEntry(
-		LOCTEXT("ToggleNodeActive", "Active"),
-		LOCTEXT("ToggleNodeActiveTooltip", "Set this node or selected nodes active/inactive"),
-		FSlateIcon(),
-		FUIAction(
-			FExecuteAction::CreateSP(&GetSequencer(), &FSequencer::ToggleNodeActive),
-			FCanExecuteAction(),
-			FIsActionChecked::CreateSP(&GetSequencer(), &FSequencer::IsNodeActive)
-		),
-		NAME_None,
-		EUserInterfaceActionType::ToggleButton
-	);
-
-	MenuBuilder.AddMenuEntry(
-		LOCTEXT("ToggleNodeLock", "Locked"),
-		LOCTEXT("ToggleNodeLockTooltip", "Lock or unlock this node or selected nodes"),
-		FSlateIcon(),
-		FUIAction(
-			FExecuteAction::CreateSP(&GetSequencer(), &FSequencer::ToggleNodeLocked),
-			FCanExecuteAction(),
-			FIsActionChecked::CreateSP(&GetSequencer(), &FSequencer::IsNodeLocked)
-		),
-		NAME_None,
-		EUserInterfaceActionType::ToggleButton
-	);
-
-	MenuBuilder.AddMenuEntry(
-		LOCTEXT("DeleteNode", "Delete"),
-		LOCTEXT("DeleteNodeTooltip", "Delete this or selected nodes"),
-		FSlateIcon(),
-		FUIAction(FExecuteAction::CreateSP(&GetSequencer(), &FSequencer::DeleteNode, ThisNode))
-	);
-
-	MenuBuilder.AddMenuEntry(
-		LOCTEXT("RenameNode", "Rename"),
-		LOCTEXT("RenameNodeTooltip", "Rename this node"),
-		FSlateIcon(),
-		FUIAction(
-			FExecuteAction::CreateSP(this, &FSequencerDisplayNode::HandleContextMenuRenameNodeExecute, ThisNode),
-			FCanExecuteAction::CreateSP(this, &FSequencerDisplayNode::HandleContextMenuRenameNodeCanExecute, ThisNode)
-		)
-	);
+	MenuBuilder.BeginSection("Edit", LOCTEXT("EditContextMenuSectionName", "Edit"));
+	{
+		MenuBuilder.AddMenuEntry(
+			LOCTEXT("ToggleNodeActive", "Active"),
+			LOCTEXT("ToggleNodeActiveTooltip", "Set this node or selected nodes active/inactive"),
+			FSlateIcon(),
+			FUIAction(
+				FExecuteAction::CreateSP(&GetSequencer(), &FSequencer::ToggleNodeActive),
+				FCanExecuteAction(),
+				FIsActionChecked::CreateSP(&GetSequencer(), &FSequencer::IsNodeActive)
+			),
+			NAME_None,
+			EUserInterfaceActionType::ToggleButton
+		);
+
+		MenuBuilder.AddMenuEntry(
+			LOCTEXT("ToggleNodeLock", "Locked"),
+			LOCTEXT("ToggleNodeLockTooltip", "Lock or unlock this node or selected nodes"),
+			FSlateIcon(),
+			FUIAction(
+				FExecuteAction::CreateSP(&GetSequencer(), &FSequencer::ToggleNodeLocked),
+				FCanExecuteAction(),
+				FIsActionChecked::CreateSP(&GetSequencer(), &FSequencer::IsNodeLocked)
+			),
+			NAME_None,
+			EUserInterfaceActionType::ToggleButton
+		);
+
+		MenuBuilder.AddMenuEntry(
+			LOCTEXT("DeleteNode", "Delete"),
+			LOCTEXT("DeleteNodeTooltip", "Delete this or selected nodes"),
+			FSlateIcon(),
+			FUIAction(FExecuteAction::CreateSP(&GetSequencer(), &FSequencer::DeleteNode, ThisNode))
+		);
+
+		MenuBuilder.AddMenuEntry(
+			LOCTEXT("RenameNode", "Rename"),
+			LOCTEXT("RenameNodeTooltip", "Rename this node"),
+			FSlateIcon(),
+			FUIAction(
+				FExecuteAction::CreateSP(this, &FSequencerDisplayNode::HandleContextMenuRenameNodeExecute, ThisNode),
+				FCanExecuteAction::CreateSP(this, &FSequencerDisplayNode::HandleContextMenuRenameNodeCanExecute, ThisNode)
+			)
+		);
+	}
+	MenuBuilder.EndSection();
 }
 
 
diff --git a/Engine/Source/Editor/Sequencer/Private/DisplayNodes/SequencerObjectBindingNode.cpp b/Engine/Source/Editor/Sequencer/Private/DisplayNodes/SequencerObjectBindingNode.cpp
--- a/Engine/Source/Editor/Sequencer/Private/DisplayNodes/SequencerObjectBindingNode.cpp
+++ b/Engine/Source/Editor/Sequencer/Private/DisplayNodes/SequencerObjectBindingNode.cpp
@@ -101,6 +101,18 @@ void FSequencerObjectBindingNode::BuildContextMenu(FMenuBuilder& MenuBuilder)
 					FNewMenuDelegate::CreateRaw(&GetSequencer(), &FSequencer::AssignActor, ObjectBinding));
 			}
 		}
+
+#if 0
+		MenuBuilder.BeginSection("Organize", LOCTEXT("OrganizeContextMenuSectionName", "Organize"));
+		{
+			MenuBuilder.AddSubMenu(
+				LOCTEXT("LabelsSubMenuText", "Labels"),
+				LOCTEXT("LabelsSubMenuTip", "Add or remove labels on this track"),
+				FNewMenuDelegate::CreateSP(this, &FSequencerObjectBindingNode::HandleLabelsSubMenuCreate)
+			);
+		}
+		MenuBuilder.EndSection();
+#endif
 	}
 
 	FSequencerDisplayNode::BuildContextMenu(MenuBuilder);
@@ -278,7 +290,7 @@ TSharedRef<SWidget> FSequencerObjectBindingNode::HandleAddTrackComboButtonGetMen
 	//@todo need to resolve this between UMG and the level editor sequencer
 	const bool bUseSubMenus = Sequencer.IsLevelEditorSequencer();
 
-	UObject* BoundObject = Sequencer.GetFocusedMovieSceneSequence()->FindObject(ObjectBinding);
+	UObject* BoundObject = Sequencer.GetFocusedMovieSceneSequenceInstance()->FindObject(ObjectBinding, Sequencer);
 
 	ISequencerModule& SequencerModule = FModuleManager::GetModuleChecked<ISequencerModule>( "Sequencer" );
 	TSharedRef<FUICommandList> CommandList(new FUICommandList);
@@ -437,10 +449,17 @@ void FSequencerObjectBindingNode::HandleAddTrackSubMenuNew(FMenuBuilder& AddTrac
 }
 
 
+void FSequencerObjectBindingNode::HandleLabelsSubMenuCreate(FMenuBuilder& MenuBuilder)
+{
+	UMovieScene* MovieScene = GetSequencer().GetFocusedMovieSceneSequence()->GetMovieScene();
+	MenuBuilder.AddWidget(SNew(SSequencerLabelEditor, MovieScene, ObjectBinding), FText::GetEmpty(), true);
+}
+
+
 void FSequencerObjectBindingNode::HandlePropertyMenuItemExecute(TArray<UProperty*> PropertyPath)
 {
 	FSequencer& Sequencer = GetSequencer();
-	UObject* BoundObject = Sequencer.GetFocusedMovieSceneSequence()->FindObject(ObjectBinding);
+	UObject* BoundObject = Sequencer.GetFocusedMovieSceneSequenceInstance()->FindObject(ObjectBinding, Sequencer);
 
 	TArray<UObject*> KeyableBoundObjects;
 	if (BoundObject != nullptr)
diff --git a/Engine/Source/Editor/Sequencer/Private/DisplayNodes/SequencerObjectBindingNode.h b/Engine/Source/Editor/Sequencer/Private/DisplayNodes/SequencerObjectBindingNode.h
--- a/Engine/Source/Editor/Sequencer/Private/DisplayNodes/SequencerObjectBindingNode.h
+++ b/Engine/Source/Editor/Sequencer/Private/DisplayNodes/SequencerObjectBindingNode.h
@@ -59,7 +59,11 @@ protected:
 private:
 
 	TSharedRef<SWidget> HandleAddTrackComboButtonGetMenuContent();
+	
 	void HandleAddTrackSubMenuNew(FMenuBuilder& AddTrackMenuBuilder, TArray<TArray<UProperty*>> KeyablePropertyPath);
+
+	void HandleLabelsSubMenuCreate(FMenuBuilder& MenuBuilder);
+	
 	void HandlePropertyMenuItemExecute(TArray<UProperty*> PropertyPath);
 
 private:
diff --git a/Engine/Source/Editor/Sequencer/Private/GroupedKeyArea.cpp b/Engine/Source/Editor/Sequencer/Private/GroupedKeyArea.cpp
--- a/Engine/Source/Editor/Sequencer/Private/GroupedKeyArea.cpp
+++ b/Engine/Source/Editor/Sequencer/Private/GroupedKeyArea.cpp
@@ -163,24 +163,19 @@ FLinearColor FGroupedKeyArea::GetKeyTint(FKeyHandle InHandle) const
 
 const FSlateBrush* FGroupedKeyArea::GetBrush(FKeyHandle InHandle) const
 {
-	static const FSlateBrush* KeyBrush = FEditorStyle::GetBrush("Sequencer.KeyDiamond");
 	static const FSlateBrush* PartialKeyBrush = FEditorStyle::GetBrush("Sequencer.PartialKey");
-
 	const auto* Group = FindGroup(InHandle);
 
-	bool bPartialKey = false;
-
 	// Ensure that each key area is represented at least once for it to be considered a 'complete key'
 	for (int32 AreaIndex = 0; Group && AreaIndex < KeyAreas.Num(); ++AreaIndex)
 	{
 		if (!Group->Keys.ContainsByPredicate([&](const FKeyGrouping::FKeyIndex& Key){ return Key.AreaIndex == AreaIndex; }))
 		{
-			bPartialKey = true;
-			break;
+			return PartialKeyBrush;
 		}
 	}
 
-	return bPartialKey ? PartialKeyBrush : KeyBrush;
+	return nullptr;
 }
 
 TArray<FKeyHandle> FGroupedKeyArea::GetUnsortedKeyHandles() const
