Copying //UE4/Dev-Sequencer to Dev-Main (//UE4/Dev-Main)

==========================
MAJOR FEATURES + CHANGES
==========================

Change 2719576 on 2015/10/07 by [REDACTED_PERSON]

	Added check for stale BP classes in FKismetCompilerUtilities::IsTypeCompatibleWithProperty() to stop compiler errors during reinstancing.
	[UE-19795] - UMG Compiler error when adding variable to nested Widget

Change 2721474 on 2015/10/08 by [REDACTED_PERSON]

	Sequencer: Movie render operations now successfully capture UMG UI

Change 2724958 on 2015/10/12 by [REDACTED_PERSON]

	Added missing resource cleanup code to UMG widgets
	[UE-21874] - UWIdget classes with missing ReleaseSlateResources() overrides

	Added ReleaseSlateResources() to ListView, TileView and Slider widgets to reset shared pointers to slate widgets.

Change 2733562 on 2015/10/19 by [REDACTED_PERSON]

	Sequencer: Fixed spawnables not working in sub-sequences

	 - The issue here was that sequence track instance updates had no knowledge of which sub-sequence they were being evaluated within. We now pass the active sequence instance into the relevant track instance functions.
	 - Also addressed some issues to do with save/restore state not getting called correctly on master tracks of sequence instances
	 - Tidied up spawn track editor

Change 2735264 on 2015/10/20 by [REDACTED_PERSON]

	Improved Engine analytics handling for Editor and games
	[UE-21892] - Improve how Engine analytics are handled for Editor and games

	Changes:
	Added Privacy section to Editor settings
	Exposed editor analytics flag in Privacy options
	Added Details Customization to make this type of bool property clearer with extra info and hyperlink
	Changed AreEditorAnalyticsEnable() to use new flag
	Prevented analytics init when disabled by user
	Sending event and shutting down analytics when user opts out
	Add in-game project setting for anonymous game usage data
	Renamed and moved bHardwareSurveyEnabled
	Added message about exposing in-game setting to end users
	Added anonymous GUID id for in-game analytics
	Moved end user settings to global config (defaultengine.ini)
	Placeholder loc text on new options for now, pending legal wording sign-off

Change 2735866 on 2015/10/20 by [REDACTED_PERSON]

	Async: Added ability to register an optional callback function that is executed when a Future completes

Change 2739793 on 2015/10/23 by [REDACTED_PERSON]

	Sequencer: Refined movie scene capturing to ensure frame accuracies are maintained

	 - Aborting an in-progress capture now gracefully terminates the process (through a remote session command) to ensure it still creates a valid video
	 - Level sequece movie capture will now pick up a corresponding level sequence in the world, and use that to capture with. A new actor will be spawned at runtime with the correct asset, should one not already exist.
	 - Made -nomovie actually work
	 - Refined how active movie captures are managed
	 - Added option to 'stage' a sequence before starting the capture. This feature will set the sequence on its first frame for the preroll, to ensure that PPP effects are allowed time to stabilize

Change 2744402 on 2015/10/28 by [REDACTED_PERSON]

	Sequencer: Separated track display names from track identifier names; code cleanup

Change 2745953 on 2015/10/29 by [REDACTED_PERSON]

	Sequencer: Attach to socket. Relative attachments.

Change 2747028 on 2015/10/29 by [REDACTED_PERSON]

	Sequencer: Another overhaul of track display name handling; code and documentation cleanup pass.

Change 2758888 on 2015/11/09 by [REDACTED_PERSON]

	Integrating changes -  4.10 to Dev-Sequencer

	From 4.10 branch fixes:
	Added check for debugger present when reporting abnormal termination to analytics. [UE-22844] CL 2750764
	Added FSystemWideCriticalSection for desktop platforms. Used by analytics to lock access to editor instances list in the OS. [UE-22844] CL 2753661
	Updating wording in privacy settings text. [UE-21892] CL 2753709
	Mac and Linux CIS fix [UE-22844] CL 2755381

Change 2761287 on 2015/11/10 by [REDACTED_PERSON]

	Sequencer: Add null check when updating the UMG preview if the sequencer doesn't exist/has been closed.

Change 2764945 on 2015/11/12 by [REDACTED_PERSON]

	Core: Templatized TypeContainer implementation to allow for thread-safe objects; updated unit test

	Also fixes UE-13850

Change 2765036 on 2015/11/12 by [REDACTED_PERSON]

	UdpMessaging: Fixed message serialization unit test (UE-22571)

Change 2766149 on 2015/11/13 by [REDACTED_PERSON]

	Media: Implemented event that gets triggered when playback reached the end of media

	Also fixes looping.

Change 2768157 on 2015/11/16 by [REDACTED_PERSON]

	Media: Added .m4a to supported WMF file extensions

Change 2769200 on 2015/11/16 by [REDACTED_PERSON]

	Editor: Add broadcast messages when snapping objects.

Change 2773066 on 2015/11/19 by [REDACTED_PERSON]

	Upload crashes from CRC to Data Router
	[UECORE-249] - Integrate Crash Report Client with the Data Router

	Upload to Receiver still active as we are running both in parallel for now.

diff --git a/Engine/Plugins/Media/AndroidMedia/Source/AndroidMedia/Private/AndroidMediaPlayer.cpp b/Engine/Plugins/Media/AndroidMedia/Source/AndroidMedia/Private/AndroidMediaPlayer.cpp
--- a/Engine/Plugins/Media/AndroidMedia/Source/AndroidMedia/Private/AndroidMediaPlayer.cpp
+++ b/Engine/Plugins/Media/AndroidMedia/Source/AndroidMedia/Private/AndroidMediaPlayer.cpp
@@ -1,21 +1,18 @@
 // Copyright 1998-2015 Epic Games, Inc. All Rights Reserved.
 
 #include "AndroidMediaPCH.h"
-
 #include "AndroidMediaPlayer.h"
 #include "AndroidApplication.h"
 #include "AndroidJava.h"
 #include "AndroidFile.h"
 #include "Paths.h"
 #include "RenderingThread.h"
-#include "RenderResource.h"
+
 
 DEFINE_LOG_CATEGORY_STATIC(LogAndroidMediaPlayer, Log, All);
 
 #define LOCTEXT_NAMESPACE "FAndroidMediaModule"
 
-// Enables directly copying into bound textures instead of using slower glReadPixels
-#define ENABLE_BINDTEXTURE	1
 
 class FAndroidMediaPlayer::MediaTrack
 	: public IMediaStream
@@ -119,6 +116,7 @@ protected:
 	TArray<IMediaSinkWeakPtr> Sinks;
 };
 
+
 class FAndroidMediaPlayer::VideoTrack
 	: public FAndroidMediaPlayer::MediaTrack
 	, public IMediaVideoTrack
@@ -185,42 +183,15 @@ public:
 		return *this;
 	}
 
-#if ENABLE_BINDTEXTURE
-	virtual void BindTexture(FTextureRHIParamRef BoundResource) override
-	{
-		BoundTextures.Add(BoundResource);
-	}
-
-	virtual void UnbindTexture(FTextureRHIParamRef BoundResource) override
-	{
-		BoundTextures.Remove(BoundResource);
-	}
-
-	virtual void Tick(float DeltaTime)
-	{
-		if (MediaPlayer.MediaState != EMediaState::Error)
-		{
-			int32 CurrentFramePosition
-				= MediaPlayer.JavaMediaPlayer->GetCurrentPosition();
-			if (LastFramePosition != CurrentFramePosition)
-			{
-				if (UpdateBoundTextures())
-				{
-					LastFramePosition = CurrentFramePosition;
-				}
-			}
-		}
-	}
-
-#else
-
 #if WITH_ENGINE
 	virtual void BindTexture(class FRHITexture* Texture) override
 	{
+		// @todo android: cbabcock: implement texture binding
 	}
 
 	virtual void UnbindTexture(class FRHITexture* Texture) override
 	{
+		// @todo android: cbabcock: implement texture binding
 	}
 #endif
 
@@ -247,39 +218,12 @@ public:
 		}
 	}
 
-#endif  // ENABLE_BINDTEXTURE
-
-private:
-
-#if ENABLE_BINDTEXTURE
-	bool UpdateBoundTextures()
-	{
-		bool FrameAvailable = false;
-		if (IsInRenderingThread())
-		{
-			for (auto Iter = BoundTextures.CreateIterator(); Iter; ++Iter)
-			{
-				FTextureRHIParamRef BoundTextureTarget = *Iter;
-				int32 DestTexture = *reinterpret_cast<int32*>(BoundTextureTarget->GetNativeResource());
-				if (MediaPlayer.JavaMediaPlayer->GetVideoLastFrame(DestTexture))
-				{
-					FrameAvailable = true;
-				}
-			}
-		}
-
-		return FrameAvailable;
-	}
-#endif  // ENABLE_BINDTEXTURE
-
 private:
 
 	int32 LastFramePosition;
-#if ENABLE_BINDTEXTURE
-	TSet<FTextureRHIParamRef> BoundTextures;
-#endif  // ENABLE_BINDTEXTURE
 };
 
+
 class FAndroidMediaPlayer::AudioTrack
 	: public FAndroidMediaPlayer::MediaTrack
 	, public IMediaAudioTrack
@@ -332,21 +276,19 @@ public:
 	}
 };
 
+
 FAndroidMediaPlayer::FAndroidMediaPlayer()
 	: JavaMediaPlayer(nullptr)
 	, MediaState(EMediaState::Error)
 {
-#if ENABLE_BINDTEXTURE
-	JavaMediaPlayer = MakeShareable(new FJavaAndroidMediaPlayer(false));
-#else
 	JavaMediaPlayer = MakeShareable(new FJavaAndroidMediaPlayer());
-#endif
 	if (JavaMediaPlayer.IsValid())
 	{
 		MediaState = EMediaState::Idle;
 	}
 }
 
+
 FAndroidMediaPlayer::~FAndroidMediaPlayer()
 {
 	Close();
@@ -357,6 +299,7 @@ FAndroidMediaPlayer::~FAndroidMediaPlayer()
 	}
 }
 
+
 FTimespan FAndroidMediaPlayer::GetDuration() const
 {
 	int32 Milliseconds = 0;
@@ -369,32 +312,38 @@ FTimespan FAndroidMediaPlayer::GetDuration() const
 	return FTimespan::FromMilliseconds(Milliseconds);
 }
 
+
 TRange<float> FAndroidMediaPlayer::GetSupportedRates(
 	EMediaPlaybackDirections Direction, bool Unthinned) const
 {
 	return TRange<float>(1.0f);
 }
 
+
 FString FAndroidMediaPlayer::GetUrl() const
 {
 	return MediaUrl;
 }
 
+
 bool FAndroidMediaPlayer::SupportsRate(float Rate, bool Unthinned) const
 {
 	return Rate == 1.0f;
 }
 
+
 bool FAndroidMediaPlayer::SupportsScrubbing() const
 {
 	return true;
 }
 
+
 bool FAndroidMediaPlayer::SupportsSeeking() const
 {
 	return true;
 }
 
+
 void FAndroidMediaPlayer::Close()
 {
 	if (MediaState == EMediaState::Prepared || MediaState == EMediaState::Started ||
@@ -406,6 +355,7 @@ void FAndroidMediaPlayer::Close()
 			JavaMediaPlayer->Stop();
 			JavaMediaPlayer->Reset();
 		}
+
 		MediaUrl = FString();
 		MediaState = EMediaState::Idle;
 		
@@ -413,11 +363,12 @@ void FAndroidMediaPlayer::Close()
 		CaptionTracks.Reset();
 		VideoTracks.Reset();
 
-		TracksChangedEvent.Broadcast();
-		ClosedEvent.Broadcast();
+		MediaEvent.Broadcast(EMediaEvent::TracksChanged);
+		MediaEvent.Broadcast(EMediaEvent::MediaClosed);
 	}
 }
 
+
 const TArray<IMediaAudioTrackRef>& FAndroidMediaPlayer::GetAudioTracks() const
 {
 	return AudioTracks;
@@ -429,16 +380,19 @@ const TArray<IMediaCaptionTrackRef>& FAndroidMediaPlayer::GetCaptionTracks() con
 	return CaptionTracks;
 }
 
+
 const IMediaInfo& FAndroidMediaPlayer::GetMediaInfo() const
 {
 	return *this;
 }
 
+
 float FAndroidMediaPlayer::GetRate() const
 {
 	return 1.0f;
 }
 
+
 FTimespan FAndroidMediaPlayer::GetTime() const
 {
 	if (MediaState != EMediaState::Error)
@@ -451,11 +405,13 @@ FTimespan FAndroidMediaPlayer::GetTime() const
 	}
 }
 
+
 const TArray<IMediaVideoTrackRef>& FAndroidMediaPlayer::GetVideoTracks() const
 {
 	return VideoTracks;
 }
 
+
 bool FAndroidMediaPlayer::IsLooping() const
 {
 	if (MediaState == EMediaState::Prepared || MediaState == EMediaState::Started ||
@@ -470,11 +426,13 @@ bool FAndroidMediaPlayer::IsLooping() const
 	}
 }
 
+
 bool FAndroidMediaPlayer::IsPaused() const
 {
 	return MediaState == EMediaState::Paused;
 }
 
+
 bool FAndroidMediaPlayer::IsPlaying() const
 {
 	if (MediaState != EMediaState::Error)
@@ -487,6 +445,7 @@ bool FAndroidMediaPlayer::IsPlaying() const
 	}
 }
 
+
 bool FAndroidMediaPlayer::IsReady() const
 {
 	return
@@ -496,6 +455,7 @@ bool FAndroidMediaPlayer::IsReady() const
 		MediaState == EMediaState::PlaybackCompleted;
 }
 
+
 bool FAndroidMediaPlayer::Open(const FString& Url)
 {
 	if (Url.IsEmpty())
@@ -526,21 +486,6 @@ bool FAndroidMediaPlayer::Open(const FString& Url)
 		FString MoviePath = Url;
 		FPaths::NormalizeFilename(MoviePath);
 
-		// Deal with hardcoded path from editor
-		if (MoviePath.Contains(":") || MoviePath.StartsWith("/"))
-		{
-			int32 Index = MoviePath.Find(TEXT("/Content/Movies"));
-			if (Index > 0)
-			{
-				Index--;
-				while (Index > 0 && MoviePath[Index] != '/')
-				{
-					Index--;
-				}
-				MoviePath = "../../.." + MoviePath.Mid(Index);
-			}
-		}
-
 		// Don't bother trying to play it if we can't find it.
 		if (!IAndroidPlatformFile::GetPlatformPhysical().FileExists(*MoviePath))
 		{
@@ -570,12 +515,14 @@ bool FAndroidMediaPlayer::Open(const FString& Url)
 			}
 		}
 	}
+
 	if (MediaState == EMediaState::Initialized)
 	{
 		MediaUrl = Url;
 		JavaMediaPlayer->Prepare();
 		MediaState = EMediaState::Prepared;
 	}
+
 	if (MediaState == EMediaState::Prepared)
 	{
 		// Use the extension as a rough guess as to what tracks
@@ -594,21 +541,25 @@ bool FAndroidMediaPlayer::Open(const FString& Url)
 			AudioTracks.Add(MakeShareable(new AudioTrack(*this, AudioTracks.Num())));
 		}
 
-		TracksChangedEvent.Broadcast();
+		MediaEvent.Broadcast(EMediaEvent::TracksChanged);
 	}
+
 	if (MediaState == EMediaState::Prepared)
 	{
-		OpenedEvent.Broadcast(MediaUrl);
+		MediaEvent.Broadcast(EMediaEvent::MediaOpenFailed);
 	}
+
 	return MediaState == EMediaState::Prepared;
 }
 
+
 bool FAndroidMediaPlayer::Open(const TSharedRef<FArchive, ESPMode::ThreadSafe>& Archive,
 	const FString& OriginalUrl)
 {
 	return false;
 }
 
+
 bool FAndroidMediaPlayer::Seek(const FTimespan& Time)
 {
 	if (MediaState == EMediaState::Prepared || MediaState == EMediaState::Started ||
@@ -617,12 +568,11 @@ bool FAndroidMediaPlayer::Seek(const FTimespan& Time)
 		JavaMediaPlayer->SeekTo(Time.GetMilliseconds());
 		return true;
 	}
-	else
-	{
-		return false;
-	}
+
+	return false;
 }
 
+
 bool FAndroidMediaPlayer::SetLooping(bool Looping)
 {
 	if (MediaState != EMediaState::Error)
@@ -630,63 +580,51 @@ bool FAndroidMediaPlayer::SetLooping(bool Looping)
 		JavaMediaPlayer->SetLooping(Looping);
 		return true;
 	}
-	else
-	{
-		return false;
-	}
+
+	return false;
 }
 
+
 bool FAndroidMediaPlayer::SetRate(float Rate)
 {
 	switch (MediaState)
 	{
 	case EMediaState::Prepared:
-		if (Rate > 0.0f)
+		if (1.0f == Rate)
 		{
 			JavaMediaPlayer->Start();
 			MediaState = EMediaState::Started;
+			MediaEvent.Broadcast(EMediaEvent::PlaybackResumed);
 			return true;
 		}
-		else
-		{
-			JavaMediaPlayer->Pause();
-			MediaState = EMediaState::Paused;
-			return true;
-		}
-		break;
-
-	case EMediaState::Started:
-		if (FMath::IsNearlyZero(Rate))
+		else if (0.0f == Rate)
 		{
 			JavaMediaPlayer->Pause();
 			MediaState = EMediaState::Paused;
+			MediaEvent.Broadcast(EMediaEvent::PlaybackSuspended);
 			return true;
 		}
 		break;
 
 	case EMediaState::Paused:
-		if (Rate > 0.0f)
-		{
-			JavaMediaPlayer->Start();
-			MediaState = EMediaState::Started;
-			return true;
-		}
-		break;
-
 	case EMediaState::PlaybackCompleted:
-		if (Rate > 0.0f)
+		if (1.0f == Rate)
 		{
 			JavaMediaPlayer->Start();
 			MediaState = EMediaState::Started;
+			MediaEvent.Broadcast(EMediaEvent::PlaybackResumed);
 			return true;
 		}
 		break;
 	}
+
 	return false;
 }
 
+
 typedef TWeakPtr<IMediaVideoTrack, ESPMode::ThreadSafe> IMediaVideoTrackWeakPtr;
 
+
 void FAndroidMediaPlayer::Tick(float DeltaTime)
 {
 	for (IMediaAudioTrackRef AudioTrack : AudioTracks)
@@ -713,11 +651,13 @@ void FAndroidMediaPlayer::Tick(float DeltaTime)
 	}
 }
 
+
 TStatId FAndroidMediaPlayer::GetStatId() const
 {
 	RETURN_QUICK_DECLARE_CYCLE_STAT(FAndroidMediaPlayer, STATGROUP_Tickables);
 }
 
+
 bool FAndroidMediaPlayer::IsTickable() const
 {
 	return true;
diff --git a/Engine/Plugins/Media/AndroidMedia/Source/AndroidMedia/Private/AndroidMediaPlayer.h b/Engine/Plugins/Media/AndroidMedia/Source/AndroidMedia/Private/AndroidMediaPlayer.h
--- a/Engine/Plugins/Media/AndroidMedia/Source/AndroidMedia/Private/AndroidMediaPlayer.h
+++ b/Engine/Plugins/Media/AndroidMedia/Source/AndroidMedia/Private/AndroidMediaPlayer.h
@@ -53,28 +53,10 @@ public:
 	virtual bool SetLooping(bool Looping) override;
 	virtual bool SetRate(float Rate) override;
 
-	DECLARE_DERIVED_EVENT(FAndroidMediaPlayer, IMediaPlayer::FOnMediaClosed, FOnMediaClosed);
-	virtual FOnMediaClosed& OnClosed() override
+	DECLARE_DERIVED_EVENT(FAndroidMediaPlayer, IMediaPlayer::FOnMediaEvent, FOnMediaEvent);
+	virtual FOnMediaEvent& OnMediaEvent() override
 	{
-		return ClosedEvent;
-	}
-
-	DECLARE_DERIVED_EVENT(FAndroidMediaPlayer, IMediaPlayer::FOnMediaOpened, FOnMediaOpened);
-	virtual FOnMediaOpened& OnOpened() override
-	{
-		return OpenedEvent;
-	}
-
-	DECLARE_DERIVED_EVENT(FAndroidMediaPlayer, IMediaPlayer::FOnMediaOpenFailed, FOnMediaOpenFailed);
-	virtual FOnMediaOpenFailed& OnOpenFailed() override
-	{
-		return OpenFailedEvent;
-	}
-
-	DECLARE_DERIVED_EVENT(FWmfMediaPlayer, IMediaPlayer::FOnTracksChanged, FOnTracksChanged);
-	virtual FOnTracksChanged& OnTracksChanged() override
-	{
-		return TracksChangedEvent;
+		return MediaEvent;
 	}
 
 public:
@@ -104,24 +86,15 @@ private:
 	/** The available caption tracks. */
 	TArray<IMediaCaptionTrackRef> CaptionTracks;
 
-	/** Holds an event delegate that is invoked when media has been closed. */
-	FOnMediaClosed ClosedEvent;
-
 	/** The Java side media interface. */
 	TSharedPtr<FJavaAndroidMediaPlayer> JavaMediaPlayer;
 
+	/** Holds an event delegate that is invoked when a media event occurred. */
+	FOnMediaEvent MediaEvent;
+
 	/** Our understanding of the state of the Java media player. */
 	EMediaState MediaState;
 
-	/** Holds an event delegate that is invoked when media has been opened. */
-	FOnMediaOpened OpenedEvent;
-
-	/** Holds an event delegate that is invoked when media failed to open. */
-	FOnMediaOpenFailed OpenFailedEvent;
-
-	/** Holds an event delegate that is invoked when the media tracks have changed. */
-	FOnTracksChanged TracksChangedEvent;
-
 	/** Currently opened media. */
 	FString MediaUrl;
 
diff --git a/Engine/Plugins/Media/AvfMedia/Source/AvfMedia/Private/Player/AvfMediaPlayer.cpp b/Engine/Plugins/Media/AvfMedia/Source/AvfMedia/Private/Player/AvfMediaPlayer.cpp
--- a/Engine/Plugins/Media/AvfMedia/Source/AvfMedia/Private/Player/AvfMediaPlayer.cpp
+++ b/Engine/Plugins/Media/AvfMedia/Source/AvfMedia/Private/Player/AvfMediaPlayer.cpp
@@ -132,8 +132,7 @@ void FAvfMediaPlayer::Close()
 	CaptionTracks.Reset();
 	VideoTracks.Reset();
 
-	TracksChangedEvent.Broadcast();
-    
+	MediaEvent.Broadcast(EMediaEvent::TracksChanged);  
     Duration = CurrentTime = FTimespan::Zero();
     
 #if PLATFORM_MAC
@@ -143,7 +142,7 @@ void FAvfMediaPlayer::Close()
     [FIOSAsyncTask CreateTaskWithBlock : ^ bool(void){
 #endif
         // Displatch on the gamethread that we have closed the video.
-        ClosedEvent.Broadcast();
+        MediaEvent.Broadcast(EMediaEvent::MediaClosed);
          
 #if PLATFORM_IOS
         return true;
@@ -265,9 +264,6 @@ bool FAvfMediaPlayer::IsReady() const
 
 bool FAvfMediaPlayer::Open( const FString& Url )
 {
-    // Cache off the url for use on other threads.
-    FString CachedUrl = Url;
-    
     Close();
     
     bool bSuccessfullyOpenedMovie = false;
@@ -320,8 +316,8 @@ bool FAvfMediaPlayer::Open( const FString& Url )
 									[FIOSAsyncTask CreateTaskWithBlock : ^ bool(void){
 #endif
 										 // Displatch on the gamethread that we have opened the video.
-										TracksChangedEvent.Broadcast();
-										OpenedEvent.Broadcast( CachedUrl );
+										MediaEvent.Broadcast(EMediaEvent::TracksChanged);
+										MediaEvent.Broadcast(EMediaEvent::MediaOpened);
 #if PLATFORM_IOS
 										return true;
 									}];
@@ -410,10 +406,20 @@ bool FAvfMediaPlayer::SetRate( float Rate )
 {
     CurrentRate = Rate;
     [MediaPlayer setRate: CurrentRate];
+
+	MediaEvent.Broadcast(
+		FMath::IsNearlyZero(Rate)
+			? EMediaEvent::PlaybackSuspended
+			: EMediaEvent::PlaybackResumed
+	);
+
     return true;
 }
 
 
+/* FAvfMediaPlayer implementation
+ *****************************************************************************/
+
 bool FAvfMediaPlayer::ReachedEnd() const
 {
 	bool bHasReachedEnd = false;
@@ -435,39 +441,42 @@ bool FAvfMediaPlayer::ShouldTick() const
 }
 
 
-/* FTickerBase interface
+/* FTickerObjectBase interface
  *****************************************************************************/
 
 bool FAvfMediaPlayer::Tick( float DeltaTime )
 {
-	if (ShouldTick())
+	if (!ShouldTick() || !IsPlaying())
+	{
+		return true;
+	}
+
+	for (IMediaVideoTrackRef& VideoTrack : VideoTracks)
 	{
-		if (IsPlaying())
+		FAvfMediaVideoTrack& AVFTrack = (FAvfMediaVideoTrack&)VideoTrack.Get();
+
+		if (!AVFTrack.ReadFrameAtTime([[MediaPlayer currentItem] currentTime]))
 		{
-			for (IMediaVideoTrackRef& VideoTrack : VideoTracks)
+			if (ReachedEnd())
 			{
-				FAvfMediaVideoTrack& AVFTrack = (FAvfMediaVideoTrack&)VideoTrack.Get();
-				if (!AVFTrack.ReadFrameAtTime([[MediaPlayer currentItem] currentTime]))
-				{
-					if(bLoop && ReachedEnd())
-					{
-						Seek(FTimespan(0));
-						SetRate(CurrentRate);
-					}
-					else
-					{
-						if(!ReachedEnd())
-						{
-							UE_LOG(LogAvfMedia, Display, TEXT("Failed to read video track for "), *MediaUrl);
-						}
-						return false;
-					}
-				}
-				else
+				MediaEvent.Broadcast(EMediaEvent::PlaybackEndReached);
+
+				if (bLoop)
 				{
-					CurrentTime = FTimespan::FromSeconds( CMTimeGetSeconds([[MediaPlayer currentItem] currentTime]) );
+					Seek(FTimespan(0));
+					SetRate(CurrentRate);
 				}
 			}
+			else
+			{
+				UE_LOG(LogAvfMedia, Display, TEXT("Failed to read video track for "), *MediaUrl);
+				
+				return false;
+			}
+		}
+		else
+		{
+			CurrentTime = FTimespan::FromSeconds( CMTimeGetSeconds([[MediaPlayer currentItem] currentTime]) );
 		}
 	}
 	
diff --git a/Engine/Plugins/Media/AvfMedia/Source/AvfMedia/Private/Player/AvfMediaPlayer.h b/Engine/Plugins/Media/AvfMedia/Source/AvfMedia/Private/Player/AvfMediaPlayer.h
--- a/Engine/Plugins/Media/AvfMedia/Source/AvfMedia/Private/Player/AvfMediaPlayer.h
+++ b/Engine/Plugins/Media/AvfMedia/Source/AvfMedia/Private/Player/AvfMediaPlayer.h
@@ -74,48 +74,31 @@ public:
 	virtual bool SetLooping( bool Looping ) override;
 	virtual bool SetRate( float Rate ) override;
 
-	DECLARE_DERIVED_EVENT(FAvfMediaPlayer, IMediaPlayer::FOnMediaClosed, FOnMediaClosed);
-	virtual FOnMediaClosed& OnClosed() override
+	DECLARE_DERIVED_EVENT(FAvfMediaPlayer, IMediaPlayer::FOnMediaEvent, FOnMediaEvent);
+	virtual FOnMediaEvent& OnMediaEvent() override
 	{
-		return ClosedEvent;
-	}
-
-	DECLARE_DERIVED_EVENT(FAvfMediaPlayer, IMediaPlayer::FOnMediaOpened, FOnMediaOpened);
-	virtual FOnMediaOpened& OnOpened() override
-	{
-		return OpenedEvent;
-	}
-
-	DECLARE_DERIVED_EVENT(FAvfMediaPlayer, IMediaPlayer::FOnMediaOpenFailed, FOnMediaOpenFailed);
-	virtual FOnMediaOpenFailed& OnOpenFailed() override
-	{
-		return OpenFailedEvent;
-	}
-
-	DECLARE_DERIVED_EVENT(FWmfMediaPlayer, IMediaPlayer::FOnTracksChanged, FOnTracksChanged);
-	virtual FOnTracksChanged& OnTracksChanged() override
-	{
-		return TracksChangedEvent;
+		return MediaEvent;
 	}
 
 public:
 
     // FTickerObjectBase interface
-    virtual bool Tick( float DeltaTime ) override;
+
+    virtual bool Tick(float DeltaTime) override;
 	
 protected:
-	
+
+	/**
+	 * Has the video completed it's playthrough
+	 */
+	bool ReachedEnd() const;
+
 	/**
 	 * Whether the media player should tick in this frame.
 	 *
 	 * @return true if the player should advance, false otherwise
 	 */
 	bool ShouldTick() const;
-	
-	/**
-	 * Has the video completed it's playthrough
-	 */
-	bool ReachedEnd() const;
 
 private:
 
@@ -143,6 +126,9 @@ private:
     /** The current playback rate. */
     float CurrentRate;
 
+	/** Holds an event delegate that is invoked when a media event occurred. */
+	FOnMediaEvent MediaEvent;
+
 	/** Cocoa helper object we can use to keep track of ns property changes in our media items */
 	FMediaHelper* MediaHelper;
 
@@ -151,18 +137,4 @@ private:
     
     /** Should the video loop to the beginning at completion */
     bool bLoop;
-
-private:
-
-	/** Holds an event delegate that is invoked when media is being unloaded. */
-	FOnMediaClosed ClosedEvent;
-
-	/** Holds an event delegate that is invoked when media has been loaded. */
-	FOnMediaOpened OpenedEvent;
-
-	/** Holds an event delegate that is invoked when media failed to open. */
-	FOnMediaOpenFailed OpenFailedEvent;
-
-	/** Holds an event delegate that is invoked when the media tracks have changed. */
-	FOnTracksChanged TracksChangedEvent;
 };
diff --git a/Engine/Plugins/Media/MediaPlayerEditor/Source/MediaPlayerEditor/Private/Widgets/SMediaPlayerEditorViewer.cpp b/Engine/Plugins/Media/MediaPlayerEditor/Source/MediaPlayerEditor/Private/Widgets/SMediaPlayerEditorViewer.cpp
--- a/Engine/Plugins/Media/MediaPlayerEditor/Source/MediaPlayerEditor/Private/Widgets/SMediaPlayerEditorViewer.cpp
+++ b/Engine/Plugins/Media/MediaPlayerEditor/Source/MediaPlayerEditor/Private/Widgets/SMediaPlayerEditorViewer.cpp
@@ -22,6 +22,7 @@ SMediaPlayerEditorViewer::~SMediaPlayerEditorViewer()
 	if (MediaPlayer != nullptr)
 	{
 		MediaPlayer->OnMediaChanged().RemoveAll(this);
+		MediaPlayer->OnTracksChanged().RemoveAll(this);
 	}
 }
 
@@ -281,8 +282,8 @@ void SMediaPlayerEditorViewer::Construct(const FArguments& InArgs, UMediaPlayer*
 			]
 	];
 
-	MediaPlayer->OnMediaChanged().AddRaw(this, &SMediaPlayerEditorViewer::HandleMediaPlayerMediaChanged);
-	MediaPlayer->OnTracksChanged().AddRaw(this, &SMediaPlayerEditorViewer::HandleMediaPlayerTracksChanged);
+	MediaPlayer->OnMediaChanged().AddSP(this, &SMediaPlayerEditorViewer::HandleMediaPlayerMediaChanged);
+	MediaPlayer->OnTracksChanged().AddSP(this, &SMediaPlayerEditorViewer::HandleMediaPlayerTracksChanged);
 
 	RegisterActiveTimer(0.0f, FWidgetActiveTimerDelegate::CreateSP(this, &SMediaPlayerEditorViewer::HandleActiveTimer));
 	ReloadTracks();
@@ -331,7 +332,7 @@ void SMediaPlayerEditorViewer::ReloadTracks()
 	}
 
 	// refresh combo box selections
-	if (!AudioTrackComboBox->GetSelectedItem().IsValid() || !AudioTracks.Contains(AudioTrackComboBox->GetSelectedItem()))
+	if (!AudioTracks.Contains(AudioTrackComboBox->GetSelectedItem()))
 	{
 		if (AudioTracks.Num() > 0)
 		{
@@ -343,7 +344,7 @@ void SMediaPlayerEditorViewer::ReloadTracks()
 		}		
 	}
 
-	if (!CaptionTrackComboBox->GetSelectedItem().IsValid() || !CaptionTracks.Contains(CaptionTrackComboBox->GetSelectedItem()))
+	if (!CaptionTracks.Contains(CaptionTrackComboBox->GetSelectedItem()))
 	{
 		if (CaptionTracks.Num() > 0)
 		{
@@ -356,7 +357,7 @@ void SMediaPlayerEditorViewer::ReloadTracks()
 		}		
 	}
 
-	if (!VideoTrackComboBox->GetSelectedItem().IsValid() || !VideoTracks.Contains(VideoTrackComboBox->GetSelectedItem()))
+	if (!VideoTracks.Contains(VideoTrackComboBox->GetSelectedItem()))
 	{
 		if (VideoTracks.Num() > 0)
 		{
@@ -498,14 +499,14 @@ FText SMediaPlayerEditorViewer::HandleOverlayStateText() const
 		return LOCTEXT("StateOverlayNoMedia", "No Media");
 	}
 
-	if (MediaPlayer->IsPaused())
+	if (!MediaPlayer->IsReady())
 	{
-		return LOCTEXT("StateOverlayPaused", "Paused");
+		return LOCTEXT("StateOverlayStopped", "Not Ready");
 	}
 
-	if (MediaPlayer->IsStopped())
+	if (MediaPlayer->IsPaused())
 	{
-		return LOCTEXT("StateOverlayStopped", "Stopped");
+		return LOCTEXT("StateOverlayPaused", "Paused");
 	}
 
 	float Rate = MediaPlayer->GetRate();
diff --git a/Engine/Plugins/Media/WmfMedia/Source/WmfMedia/Private/Player/WmfMediaPlayer.cpp b/Engine/Plugins/Media/WmfMedia/Source/WmfMedia/Private/Player/WmfMediaPlayer.cpp
--- a/Engine/Plugins/Media/WmfMedia/Source/WmfMedia/Private/Player/WmfMediaPlayer.cpp
+++ b/Engine/Plugins/Media/WmfMedia/Source/WmfMedia/Private/Player/WmfMediaPlayer.cpp
@@ -16,15 +16,17 @@ FWmfMediaPlayer::FWmfMediaPlayer()
 	{
 		Resolver->OnResolveComplete().BindLambda([=](TComPtr<IUnknown> SourceObject, FString ResolvedUrl) {
 			AsyncTask(ENamedThreads::GameThread, [=]() {
-				InitializeMediaSession(SourceObject, ResolvedUrl)
-					? OpenedEvent.Broadcast(ResolvedUrl)
-					: OpenFailedEvent.Broadcast(ResolvedUrl);
+				MediaEvent.Broadcast(
+					InitializeMediaSession(SourceObject, ResolvedUrl)
+						? EMediaEvent::MediaOpened
+						: EMediaEvent::MediaOpenFailed
+				);
 			});
 		});
 
 		Resolver->OnResolveFailed().BindLambda([=](FString FailedUrl) {
 			AsyncTask(ENamedThreads::GameThread, [=]() {
-				OpenFailedEvent.Broadcast(FailedUrl);
+				MediaEvent.Broadcast(EMediaEvent::MediaOpenFailed);
 			});
 		});
 	}
@@ -64,7 +66,7 @@ FString FWmfMediaPlayer::GetUrl() const
 }
 
 
-bool FWmfMediaPlayer::SupportsRate( float Rate, bool Unthinned ) const
+bool FWmfMediaPlayer::SupportsRate(float Rate, bool Unthinned) const
 {
 	return (MediaSession != NULL)
 		? MediaSession->IsRateSupported(Rate, Unthinned)
@@ -98,6 +100,14 @@ void FWmfMediaPlayer::Close()
 		return;
 	}
 
+	MediaSession->OnError().RemoveAll(this);
+	MediaSession->OnSessionEvent().RemoveAll(this);
+
+	if (IsPlaying())
+	{
+		MediaEvent.Broadcast(EMediaEvent::PlaybackSuspended);
+	}
+
 	MediaSession->SetState(EMediaStates::Closed);
 	MediaSession.Reset();
 
@@ -114,8 +124,8 @@ void FWmfMediaPlayer::Close()
 	Duration = 0;
 	MediaUrl = FString();
 
-	TracksChangedEvent.Broadcast();
-	ClosedEvent.Broadcast();
+	MediaEvent.Broadcast(EMediaEvent::TracksChanged);
+	MediaEvent.Broadcast(EMediaEvent::MediaClosed);
 }
 
 
@@ -243,17 +253,7 @@ bool FWmfMediaPlayer::SetRate(float Rate)
 		return false;
 	}
 
-	if (FMath::IsNearlyZero(Rate))
-	{
-		if (MediaSession->SupportsScrubbing())
-		{
-			return MediaSession->SetRate(Rate);
-		}
-
-		return MediaSession->SetState(EMediaStates::Paused);
-	}
-
-	return MediaSession->SetRate(Rate) && MediaSession->SetState(EMediaStates::Playing);
+	return MediaSession->SetRate(Rate);
 }
 
 
@@ -454,7 +454,7 @@ bool FWmfMediaPlayer::InitializeMediaSession(IUnknown* SourceObject, const FStri
 	}
 
 	UE_LOG(LogWmfMedia, Verbose, TEXT("Added a total of %i audio tracks, %i caption tracks, %i video tracks"), AudioTracks.Num(), CaptionTracks.Num(), VideoTracks.Num());
-	TracksChangedEvent.Broadcast();
+	MediaEvent.Broadcast(EMediaEvent::TracksChanged);
 
 	UINT64 PresentationDuration = 0;
 	PresentationDescriptor->GetUINT64(MF_PD_DURATION, &PresentationDuration);
@@ -466,6 +466,7 @@ bool FWmfMediaPlayer::InitializeMediaSession(IUnknown* SourceObject, const FStri
 	MediaSession = new FWmfMediaSession(Duration, Topology);
 	{
 		MediaSession->OnError().AddRaw(this, &FWmfMediaPlayer::HandleSessionError);
+		MediaSession->OnSessionEvent().AddRaw(this, &FWmfMediaPlayer::HandleSessionEvent);
 	}
 
 	return (MediaSession->GetState() != EMediaStates::Error);
@@ -477,7 +478,36 @@ bool FWmfMediaPlayer::InitializeMediaSession(IUnknown* SourceObject, const FStri
 
 void FWmfMediaPlayer::HandleSessionError(HRESULT Error)
 {
-	UE_LOG(LogWmfMedia, Error, TEXT("An unrecoverable error occured in the media session: 0x%X"), Error);
+	UE_LOG(LogWmfMedia, Error, TEXT("An error occured in the media session: 0x%X"), Error);
+}
+
+
+void FWmfMediaPlayer::HandleSessionEvent(MediaEventType EventType)
+{
+	EMediaEvent Event = EMediaEvent::Unknown;
+
+	switch (EventType)
+	{
+	case MEEndOfPresentation:
+		Event = EMediaEvent::PlaybackEndReached;
+		break;
+
+	case MESessionStarted:
+		Event = EMediaEvent::PlaybackResumed;
+		break;
+
+	case MESessionStopped:
+		Event = EMediaEvent::PlaybackSuspended;
+		break;
+	}
+
+	if (Event != EMediaEvent::Unknown)
+	{
+		// forward event to game thread
+		AsyncTask(ENamedThreads::GameThread, [=]() {
+			MediaEvent.Broadcast(Event);
+		});
+	}
 }
 
 
diff --git a/Engine/Plugins/Media/WmfMedia/Source/WmfMedia/Private/Player/WmfMediaPlayer.h b/Engine/Plugins/Media/WmfMedia/Source/WmfMedia/Private/Player/WmfMediaPlayer.h
--- a/Engine/Plugins/Media/WmfMedia/Source/WmfMedia/Private/Player/WmfMediaPlayer.h
+++ b/Engine/Plugins/Media/WmfMedia/Source/WmfMedia/Private/Player/WmfMediaPlayer.h
@@ -52,28 +52,10 @@ public:
 	virtual bool SetLooping(bool Looping) override;
 	virtual bool SetRate(float Rate) override;
 
-	DECLARE_DERIVED_EVENT(FWmfMediaPlayer, IMediaPlayer::FOnMediaClosed, FOnMediaClosed);
-	virtual FOnMediaClosed& OnClosed() override
+	DECLARE_DERIVED_EVENT(FWmfMediaPlayer, IMediaPlayer::FOnMediaEvent, FOnMediaEvent);
+	virtual FOnMediaEvent& OnMediaEvent() override
 	{
-		return ClosedEvent;
-	}
-
-	DECLARE_DERIVED_EVENT(FWmfMediaPlayer, IMediaPlayer::FOnMediaOpened, FOnMediaOpened);
-	virtual FOnMediaOpened& OnOpened() override
-	{
-		return OpenedEvent;
-	}
-
-	DECLARE_DERIVED_EVENT(FWmfMediaPlayer, IMediaPlayer::FOnMediaOpenFailed, FOnMediaOpenFailed);
-	virtual FOnMediaOpenFailed& OnOpenFailed() override
-	{
-		return OpenFailedEvent;
-	}
-
-	DECLARE_DERIVED_EVENT(FWmfMediaPlayer, IMediaPlayer::FOnTracksChanged, FOnTracksChanged);
-	virtual FOnTracksChanged& OnTracksChanged() override
-	{
-		return TracksChangedEvent;
+		return MediaEvent;
 	}
 
 protected:
@@ -103,6 +85,9 @@ private:
 	/** Handles session errors. */
 	void HandleSessionError(HRESULT Error);
 
+	/** Handles session events. */
+	void HandleSessionEvent(MediaEventType EventType);
+
 private:
 
 	/** The available audio tracks. */
@@ -114,6 +99,9 @@ private:
 	/** The duration of the currently loaded media. */
 	FTimespan Duration;
 
+	/** Holds an event delegate that is invoked when a media event occurred. */
+	FOnMediaEvent MediaEvent;
+
 	/** Holds the asynchronous callback object for the media stream. */
 	TComPtr<FWmfMediaSession> MediaSession;
 
@@ -128,20 +116,6 @@ private:
 
 	/** The available video tracks. */
 	TArray<IMediaVideoTrackRef> VideoTracks;
-
-private:
-
-	/** Holds an event delegate that is invoked when media has been closed. */
-	FOnMediaClosed ClosedEvent;
-
-	/** Holds an event delegate that is invoked when media has been opened. */
-	FOnMediaOpened OpenedEvent;
-
-	/** Holds an event delegate that is invoked when media failed to open. */
-	FOnMediaOpenFailed OpenFailedEvent;
-
-	/** Holds an event delegate that is invoked when the media tracks have changed. */
-	FOnTracksChanged TracksChangedEvent;
 };
 
 
diff --git a/Engine/Plugins/Media/WmfMedia/Source/WmfMedia/Private/Wmf/WmfMediaSession.cpp b/Engine/Plugins/Media/WmfMedia/Source/WmfMedia/Private/Wmf/WmfMediaSession.cpp
--- a/Engine/Plugins/Media/WmfMedia/Source/WmfMedia/Private/Wmf/WmfMediaSession.cpp
+++ b/Engine/Plugins/Media/WmfMedia/Source/WmfMedia/Private/Wmf/WmfMediaSession.cpp
@@ -7,7 +7,7 @@
 /* FWmfMediaSession structors
  *****************************************************************************/
 
-FWmfMediaSession::FWmfMediaSession( const FTimespan& InDuration, const TComPtr<IMFTopology>& InTopology )
+FWmfMediaSession::FWmfMediaSession(const FTimespan& InDuration, const TComPtr<IMFTopology>& InTopology)
 	: Capabilities(0)
 	, ChangeRequested(false)
 	, CurrentRate(0.0f)
@@ -75,7 +75,7 @@ FTimespan FWmfMediaSession::GetPosition() const
 }
 
 
-TRange<float> FWmfMediaSession::GetSupportedRates( EMediaPlaybackDirections Direction, bool Unthinned ) const
+TRange<float> FWmfMediaSession::GetSupportedRates(EMediaPlaybackDirections Direction, bool Unthinned) const
 {
 	if (RateSupport == NULL)
 	{
@@ -102,7 +102,7 @@ TRange<float> FWmfMediaSession::GetSupportedRates( EMediaPlaybackDirections Dire
 }
 
 
-bool FWmfMediaSession::IsRateSupported( float Rate, bool Unthinned ) const
+bool FWmfMediaSession::IsRateSupported(float Rate, bool Unthinned) const
 {
 	if (Rate == 1.0f)
 	{
@@ -123,7 +123,7 @@ bool FWmfMediaSession::IsRateSupported( float Rate, bool Unthinned ) const
 }
 
 
-bool FWmfMediaSession::SetPosition( const FTimespan& Position )
+bool FWmfMediaSession::SetPosition(const FTimespan& Position)
 {
 	RequestedPosition = Position;
 
@@ -131,7 +131,7 @@ bool FWmfMediaSession::SetPosition( const FTimespan& Position )
 }
 
 
-bool FWmfMediaSession::SetRate( float Rate )
+bool FWmfMediaSession::SetRate(float Rate)
 {
 	if (!IsRateSupported(Rate, false))
 	{
@@ -142,7 +142,10 @@ bool FWmfMediaSession::SetRate( float Rate )
 
 	if (FMath::IsNearlyZero(Rate))
 	{
-		RequestedState = EMediaStates::Paused;
+		if (!SupportsScrubbing())
+		{
+			RequestedState = EMediaStates::Paused;
+		}
 	}
 	else
 	{
@@ -153,7 +156,7 @@ bool FWmfMediaSession::SetRate( float Rate )
 }
 
 
-bool FWmfMediaSession::SetState( EMediaStates NewState )
+bool FWmfMediaSession::SetState(EMediaStates NewState)
 {
 	RequestedState = NewState;
 
@@ -176,13 +179,13 @@ STDMETHODIMP_(ULONG) FWmfMediaSession::AddRef()
 }
 
 
-STDMETHODIMP FWmfMediaSession::GetParameters( unsigned long*, unsigned long*)
+STDMETHODIMP FWmfMediaSession::GetParameters(unsigned long*, unsigned long*)
 {
 	return E_NOTIMPL;
 }
 
 
-STDMETHODIMP FWmfMediaSession::Invoke( IMFAsyncResult* AsyncResult )
+STDMETHODIMP FWmfMediaSession::Invoke(IMFAsyncResult* AsyncResult)
 {
 	TComPtr<IMFMediaEvent> Event;
 	
@@ -198,93 +201,104 @@ STDMETHODIMP FWmfMediaSession::Invoke( IMFAsyncResult* AsyncResult )
 		return S_OK;
 	}
 
-		HRESULT EventResult;
-		Event->GetStatus(&EventResult);
+	HRESULT EventResult;
+	Event->GetStatus(&EventResult);
 
-		FScopeLock ScopeLock(&CriticalSection);
+	FScopeLock ScopeLock(&CriticalSection);
+	{
+		if (EventType == MESessionClosed)
 		{
-			if (EventType == MESessionClosed)
-			{
-				UpdateState(EMediaStates::Closed);
-				Capabilities = 0;
-			}
-			else if (EventType == MEError)
-			{
-				ErrorEvent.Broadcast(EventResult);
-				UpdateState(EMediaStates::Error);
-			}
-			else
+			UpdateState(EMediaStates::Closed);
+			Capabilities = 0;
+		}
+		else if (EventType == MEError)
+		{
+			ErrorEvent.Broadcast(EventResult);
+			UpdateState(EMediaStates::Error);
+		}
+		else
+		{
+			switch (EventType)
 			{
-				if (EventType == MESessionCapabilitiesChanged)
+			case MEEndOfPresentation:
+				if (Looping && (CurrentRate > 0.0f))
 				{
-					Capabilities = ::MFGetAttributeUINT32(Event, MF_EVENT_SESSIONCAPS, Capabilities);
+					RequestedPosition = FTimespan::Zero();
+					ChangeState();
 				}
-				else if (EventType == MESessionEnded)
-				{
-					if (Looping)
-					{
-						if (CurrentRate < 0.0f)
-						{
-							RequestedPosition = Duration;
-						}
+				break;
 
-						UpdateState(EMediaStates::Paused);
-					}
-					else
+			case MESessionCapabilitiesChanged:
+				Capabilities = ::MFGetAttributeUINT32(Event, MF_EVENT_SESSIONCAPS, Capabilities);
+				break;
+
+			case MESessionEnded:
+				if (Looping)
+				{
+					if (CurrentRate < 0.0f)
 					{
-						RequestedState = EMediaStates::Stopped;
-						UpdateState(EMediaStates::Stopped);
+						RequestedPosition = Duration;
+						ChangeState();
 					}
 				}
-				else if (EventType == MESessionPaused)
+				else
 				{
-					UpdateState(EMediaStates::Paused);
+					RequestedState = EMediaStates::Stopped;
+					UpdateState(EMediaStates::Stopped);
 				}
-				else if (EventType == MESessionRateChanged)
-				{
-					// recover active playback rate if rate change failed
-					if (FAILED(EventResult) && (CurrentRate == RequestedRate))
-					{
-						PROPVARIANT Value;
-						PropVariantInit(&Value);
+				break;
 
-						if (SUCCEEDED(Event->GetValue(&Value)) && (Value.vt == VT_R4))
-						{
-							CurrentRate = Value.fltVal;
-						}
+			case MESessionPaused:
+				UpdateState(EMediaStates::Paused);
+				break;
 
-						RequestedRate = CurrentRate;
-					}
-				}
-				else if (EventType == MESessionScrubSampleComplete)
-				{
-					if (GetInternalPosition() == RequestedPosition)
-					{
-						RequestedPosition = FTimespan::MinValue();
-					}
-				}
-				else if (EventType == MESessionStarted)
+			case MESessionRateChanged:
+				// recover active playback rate if rate change failed
+				if (FAILED(EventResult) && (CurrentRate == RequestedRate))
 				{
-					if (GetInternalPosition() == RequestedPosition)
+					PROPVARIANT Value;
+					PropVariantInit(&Value);
+
+					if (SUCCEEDED(Event->GetValue(&Value)) && (Value.vt == VT_R4))
 					{
-						RequestedPosition = FTimespan::MinValue();
+						CurrentRate = Value.fltVal;
 					}
 
-					UpdateState(EMediaStates::Playing);
+					RequestedRate = CurrentRate;
 				}
-				else if (EventType == MESessionStopped)
+
+			case MESessionScrubSampleComplete:
+				if (GetInternalPosition() == RequestedPosition)
 				{
-					UpdateState(EMediaStates::Stopped);
+					RequestedPosition = FTimespan::MinValue();
 				}
+				break;
 
-				// request the next event
-				if (FAILED(MediaSession->BeginGetEvent(this, NULL)))
+			case MESessionStarted:
+				if (GetInternalPosition() == RequestedPosition)
 				{
-					Capabilities = 0;
-					CurrentState = EMediaStates::Error;
+					RequestedPosition = FTimespan::MinValue();
 				}
+
+				UpdateState(EMediaStates::Playing);
+				break;
+
+			case MESessionStopped:
+				UpdateState(EMediaStates::Stopped);
+				break;
+			}
+
+			// bubble up event to player
+			SessionEvent.Broadcast(EventType);
+
+			// request the next event
+			if (FAILED(MediaSession->BeginGetEvent(this, NULL)))
+			{
+				Capabilities = 0;
+				CurrentState = EMediaStates::Error;
 			}
 		}
+	}
 	
 	return S_OK;
 }
@@ -295,7 +309,7 @@ STDMETHODIMP FWmfMediaSession::Invoke( IMFAsyncResult* AsyncResult )
 #pragma warning(disable:4838)
 #endif // _MSC_VER == 1900
 
-STDMETHODIMP FWmfMediaSession::QueryInterface( REFIID RefID, void** Object )
+STDMETHODIMP FWmfMediaSession::QueryInterface(REFIID RefID, void** Object)
 {
 	static const QITAB QITab[] =
 	{
@@ -334,13 +348,20 @@ STDMETHODIMP_(ULONG) FWmfMediaSession::Release()
 bool FWmfMediaSession::ChangeState()
 {
 	// disallow state changes if session was closed or had an error
-	if ((CurrentState == EMediaStates::Closed) || (CurrentState == EMediaStates::Error))
+	if (CurrentState == EMediaStates::Closed)
 	{
 		return false;
 	}
 
 	FScopeLock ScopeLock(&CriticalSection);
 	{
+		if (CurrentState == EMediaStates::Error)
+		{
+			MediaSession->Close();
+
+			return true;
+		}
+
 		// defer state change
 		if (StateChangePending)
 		{
@@ -350,7 +371,7 @@ bool FWmfMediaSession::ChangeState()
 		}
 
 		// close session
-		if ((RequestedState == EMediaStates::Closed) || (RequestedState == EMediaStates::Error))
+		if (RequestedState == EMediaStates::Closed)
 		{
 			MediaSession->Close();
 
@@ -460,7 +481,7 @@ FTimespan FWmfMediaSession::GetInternalPosition() const
 }
 
 
-void FWmfMediaSession::UpdateState( EMediaStates CompletedState )
+void FWmfMediaSession::UpdateState(EMediaStates CompletedState)
 {
 	CurrentState = CompletedState;
 
diff --git a/Engine/Plugins/Media/WmfMedia/Source/WmfMedia/Private/Wmf/WmfMediaSession.h b/Engine/Plugins/Media/WmfMedia/Source/WmfMedia/Private/Wmf/WmfMediaSession.h
--- a/Engine/Plugins/Media/WmfMedia/Source/WmfMedia/Private/Wmf/WmfMediaSession.h
+++ b/Engine/Plugins/Media/WmfMedia/Source/WmfMedia/Private/Wmf/WmfMediaSession.h
@@ -54,7 +54,7 @@ public:
 	 * @param InDuration The duration of the media.
 	 * @param InTopology The playback topology to use.
 	 */
-	FWmfMediaSession( const FTimespan& InDuration, const TComPtr<IMFTopology>& InTopology );
+	FWmfMediaSession(const FTimespan& InDuration, const TComPtr<IMFTopology>& InTopology);
 
 	/** Virtual destructor. */
 	virtual ~FWmfMediaSession() { }
@@ -179,6 +179,13 @@ public:
 		return ErrorEvent;
 	}
 
+	/** Gets an event delegate that is invoked when the last presentation segment finished playing. */
+	DECLARE_EVENT_OneParam(FWmfMediaSession, FOnSessionEvent, MediaEventType /*EventType*/)
+	FOnSessionEvent& OnSessionEvent()
+	{
+		return SessionEvent;
+	}
+
 public:
 
 	// IMFAsyncCallback interface
@@ -269,6 +276,9 @@ private:
 
 	/** Holds an event delegate that is invoked when an error occurred. */
 	FOnError ErrorEvent;
+
+	/** Holds an event delegate that is invoked when an event occurs in the session. */
+	FOnSessionEvent SessionEvent;
 };
 
 
diff --git a/Engine/Plugins/Media/WmfMedia/Source/WmfMedia/Private/WmfMediaModule.cpp b/Engine/Plugins/Media/WmfMedia/Source/WmfMedia/Private/WmfMediaModule.cpp
--- a/Engine/Plugins/Media/WmfMedia/Source/WmfMedia/Private/WmfMediaModule.cpp
+++ b/Engine/Plugins/Media/WmfMedia/Source/WmfMedia/Private/WmfMediaModule.cpp
@@ -74,6 +74,7 @@ public:
 		SupportedFileTypes.Add(TEXT("adts"), LOCTEXT("FormatAdts", "Audio Data Transport Stream"));
 		SupportedFileTypes.Add(TEXT("asf"), LOCTEXT("FormatAsf", "ASF Media File"));
 		SupportedFileTypes.Add(TEXT("avi"), LOCTEXT("FormatAvi", "Audio Video Interleave File"));
+		SupportedFileTypes.Add(TEXT("m4a"), LOCTEXT("FormatM4a", "Apple MPEG-4 Audio"));
 		SupportedFileTypes.Add(TEXT("m4v"), LOCTEXT("FormatM4v", "Apple MPEG-4 Video"));
 		SupportedFileTypes.Add(TEXT("mov"), LOCTEXT("FormatMov", "Apple QuickTime Movie"));
 		SupportedFileTypes.Add(TEXT("mp3"), LOCTEXT("FormatMp3", "MPEG-2 Audio"));
diff --git a/Engine/Plugins/Messaging/MessagingDebugger/MessagingDebugger.uplugin b/Engine/Plugins/Messaging/MessagingDebugger/MessagingDebugger.uplugin
--- a/Engine/Plugins/Messaging/MessagingDebugger/MessagingDebugger.uplugin
+++ b/Engine/Plugins/Messaging/MessagingDebugger/MessagingDebugger.uplugin
@@ -1,25 +1,24 @@
 {
 	"FileVersion" : 3,
-	"Version" : 1,
-	"VersionName" : "1.0",
+	
 	"FriendlyName" : "Messaging Debugger",
-	"Description" : "Provides a visual debugger for the messaging sub-system.",
-	"Category" : "Messaging",
+	"Version" : 1,
+	"VersionName" : "1.1",
 	"CreatedBy" : "Epic Games, Inc.",
 	"CreatedByURL" : "http://epicgames.com",
-	"DocsURL" : "",
-	"MarketplaceURL" : "",
-	"SupportURL" : "",
-	"EnabledByDefault" : true,
-	"CanContainContent" : false,
-	"IsBetaVersion" : false,
-	"Installed" : false,
+	"EngineVersion" : "4.10.0",
+	"Description" : "Provides a visual debugger for the messaging sub-system.",
+	"Category" : "Messaging",
+	"EnabledByDefault" : false,
+	"IsBetaVersion": true,
+
 	"Modules" :
 	[
 		{
 			"Name" : "MessagingDebugger",
-			"Type" : "Editor",
-			"LoadingPhase" : "Default"
+			"Type" : "Developer"
 		}
-	]
+	],
+
+	"CanContainContent" : false
 }
\ No newline at end of file
diff --git a/Engine/Plugins/Messaging/MessagingDebugger/Source/MessagingDebugger/MessagingDebugger.Build.cs b/Engine/Plugins/Messaging/MessagingDebugger/Source/MessagingDebugger/MessagingDebugger.Build.cs
--- a/Engine/Plugins/Messaging/MessagingDebugger/Source/MessagingDebugger/MessagingDebugger.Build.cs
+++ b/Engine/Plugins/Messaging/MessagingDebugger/Source/MessagingDebugger/MessagingDebugger.Build.cs
@@ -45,7 +45,12 @@ public class MessagingDebugger : ModuleRules
 
         if (UEBuildConfiguration.bBuildEditor)
         {
-            PrivateDependencyModuleNames.Add("PropertyEditor");
+            PrivateDependencyModuleNames.AddRange(
+                new string[] {
+                    "PropertyEditor",
+                    "WorkspaceMenuStructure",
+                }
+            );
         }
 	}
 }
diff --git a/Engine/Plugins/Messaging/MessagingDebugger/Source/MessagingDebugger/Private/MessagingDebuggerModule.cpp b/Engine/Plugins/Messaging/MessagingDebugger/Source/MessagingDebugger/Private/MessagingDebuggerModule.cpp
--- a/Engine/Plugins/Messaging/MessagingDebugger/Source/MessagingDebugger/Private/MessagingDebuggerModule.cpp
+++ b/Engine/Plugins/Messaging/MessagingDebugger/Source/MessagingDebugger/Private/MessagingDebuggerModule.cpp
@@ -4,6 +4,10 @@
 #include "Runtime/Core/Public/Features/IModularFeatures.h"
 #include "SDockTab.h"
 
+#if WITH_EDITOR
+	#include "WorkspaceMenuStructureModule.h"
+#endif
+
 
 #define LOCTEXT_NAMESPACE "FMessagingDebuggerModule"
 
@@ -27,18 +31,21 @@ public:
 		Style = MakeShareable(new FMessagingDebuggerStyle());
 
 		FMessagingDebuggerCommands::Register();
-
 		IModularFeatures::Get().RegisterModularFeature("MessagingDebugger", this);
-		
-		// This is still experimental in the editor, so it'll be invoked specifically in FMainMenu if the experimental settings flag is set.
-		// When no longer experimental, switch to the nomad spawner registration below
-		FGlobalTabmanager::Get()->RegisterTabSpawner(MessagingDebuggerTabName, FOnSpawnTab::CreateRaw(this, &FMessagingDebuggerModule::SpawnMessagingDebuggerTab));
+		FGlobalTabmanager::Get()->RegisterNomadTabSpawner(MessagingDebuggerTabName, FOnSpawnTab::CreateRaw(this, &FMessagingDebuggerModule::SpawnMessagingDebuggerTab))
+			.SetDisplayName(LOCTEXT("TabTitle", "Messaging Debugger"))
+#if WITH_EDITOR
+			.SetGroup(WorkspaceMenu::GetMenuStructure().GetDeveloperToolsDebugCategory())
+#else
+
+#endif
+			.SetIcon(FSlateIcon(Style->GetStyleSetName(), "MessagingDebuggerTabIcon"))
+			.SetTooltipText(LOCTEXT("TooltipText", "Visual debugger for the messaging sub-system."));
 	}
 
 	virtual void ShutdownModule() override
 	{
-		FGlobalTabmanager::Get()->UnregisterTabSpawner(MessagingDebuggerTabName);
-
+		FGlobalTabmanager::Get()->UnregisterNomadTabSpawner(MessagingDebuggerTabName);
 		IModularFeatures::Get().UnregisterModularFeature("MessagingDebugger", this);
 		FMessagingDebuggerCommands::Unregister();
 	}
@@ -51,7 +58,7 @@ private:
 	 * @param SpawnTabArgs The arguments for the tab to spawn.
 	 * @return The spawned tab.
 	 */
-	TSharedRef<SDockTab> SpawnMessagingDebuggerTab( const FSpawnTabArgs& SpawnTabArgs )
+	TSharedRef<SDockTab> SpawnMessagingDebuggerTab(const FSpawnTabArgs& SpawnTabArgs)
 	{
 		const TSharedRef<SDockTab> MajorTab = SNew(SDockTab)
 			.TabRole(ETabRole::MajorTab);
diff --git a/Engine/Plugins/Messaging/UdpMessaging/Source/UdpMessaging/Private/Tests/UdpSerializeMessageTaskTest.cpp b/Engine/Plugins/Messaging/UdpMessaging/Source/UdpMessaging/Private/Tests/UdpSerializeMessageTaskTest.cpp
--- a/Engine/Plugins/Messaging/UdpMessaging/Source/UdpMessaging/Private/Tests/UdpSerializeMessageTaskTest.cpp
+++ b/Engine/Plugins/Messaging/UdpMessaging/Source/UdpMessaging/Private/Tests/UdpSerializeMessageTaskTest.cpp
@@ -52,10 +52,9 @@ bool FUdpSerializeMessageTaskTest::RunTest(const FString& Parameters)
 
 	CompletedMessages = 0;
 	FailedMessages = 0;
-	ReferenceMessage = MakeShareable(new FUdpSerializedMessage());
+	ReferenceMessage = MakeShareable(new FUdpSerializedMessage);
 
-	FUdpMockMessage* Message = new FUdpMockMessage();
-	IMessageContextRef Context = MakeShareable(new FUdpMockMessageContext(Message));
+	IMessageContextRef Context = MakeShareable(new FUdpMockMessageContext(new FUdpMockMessage));
 
 	// synchronous reference serialization
 	FUdpSerializeMessageTask ReferenceTask(Context, ReferenceMessage.ToSharedRef());
@@ -66,7 +65,7 @@ bool FUdpSerializeMessageTaskTest::RunTest(const FString& Parameters)
 	// stress test
 	for (int32 TestIndex = 0; TestIndex < NumMessages; ++TestIndex)
 	{
-		FUdpSerializedMessageRef SerializedMessage = MakeShareable(new FUdpSerializedMessage());
+		FUdpSerializedMessageRef SerializedMessage = MakeShareable(new FUdpSerializedMessage);
 		{
 			SerializedMessage->OnStateChanged().BindStatic(&HandleSerializedMessageStateChanged, SerializedMessage);
 		}
@@ -84,9 +83,5 @@ bool FUdpSerializeMessageTaskTest::RunTest(const FString& Parameters)
 	TestEqual(TEXT("The number of completed messages must equal the total number of messages"), CompletedMessages, NumMessages);
 	TestEqual(TEXT("There must be no failed messages"), FailedMessages, (int32)0);
 
-	// clean up
-	FMemory::Free(Message);
-	ReferenceMessage.Reset();
-
 	return ((CompletedMessages == NumMessages) && (FailedMessages == 0));
 }
diff --git a/Engine/Plugins/Messaging/UdpMessaging/Source/UdpMessaging/Private/Transport/UdpMessageProcessor.cpp b/Engine/Plugins/Messaging/UdpMessaging/Source/UdpMessaging/Private/Transport/UdpMessageProcessor.cpp
--- a/Engine/Plugins/Messaging/UdpMessaging/Source/UdpMessaging/Private/Transport/UdpMessageProcessor.cpp
+++ b/Engine/Plugins/Messaging/UdpMessaging/Source/UdpMessaging/Private/Transport/UdpMessageProcessor.cpp
@@ -46,8 +46,8 @@ FUdpMessageProcessor::FUdpMessageProcessor(FSocket* InSocket, const FGuid& InNod
 FUdpMessageProcessor::~FUdpMessageProcessor()
 {
 	Thread->Kill(true);
-		
 	delete Thread;
+
 	FPlatformProcess::ReturnSynchEventToPool(WorkEvent);
 	WorkEvent = nullptr;
 }
@@ -104,7 +104,6 @@ uint32 FUdpMessageProcessor::Run()
 
 			ConsumeInboundSegments();
 			ConsumeOutboundMessages();
-
 			UpdateKnownNodes();
 			UpdateStaticNodes();
 		}
@@ -151,9 +150,8 @@ void FUdpMessageProcessor::AcknowledgeReceipt(int32 MessageId, const FNodeInfo&
 		Writer << AcknowledgeChunk;
 	}
 
-	int32 Sent;
-
-	Socket->SendTo(Writer.GetData(), Writer.Num(), Sent, *NodeInfo.Endpoint.ToInternetAddr());
+	int32 OutSent;
+	Socket->SendTo(Writer.GetData(), Writer.Num(), OutSent, *NodeInfo.Endpoint.ToInternetAddr());
 }
 
 
@@ -288,8 +286,8 @@ bool FUdpMessageProcessor::FilterSegment(const FUdpMessageSegment::FHeader& Head
 void FUdpMessageProcessor::ProcessAbortSegment(FInboundSegment& Segment, FNodeInfo& NodeInfo)
 {
 	FUdpMessageSegment::FAbortChunk AbortChunk;
-
 	*Segment.Data << AbortChunk;
+
 	NodeInfo.Segmenters.Remove(AbortChunk.MessageId);
 }
 
@@ -297,8 +295,8 @@ void FUdpMessageProcessor::ProcessAbortSegment(FInboundSegment& Segment, FNodeIn
 void FUdpMessageProcessor::ProcessAcknowledgeSegment(FInboundSegment& Segment, FNodeInfo& NodeInfo)
 {
 	FUdpMessageSegment::FAcknowledgeChunk AcknowledgeChunk;
-
 	*Segment.Data << AcknowledgeChunk;
+
 	NodeInfo.Segmenters.Remove(AcknowledgeChunk.MessageId);
 }
 
@@ -306,7 +304,6 @@ void FUdpMessageProcessor::ProcessAcknowledgeSegment(FInboundSegment& Segment, F
 void FUdpMessageProcessor::ProcessByeSegment(FInboundSegment& Segment, FNodeInfo& NodeInfo)
 {
 	FGuid RemoteNodeId;
-
 	*Segment.Data << RemoteNodeId;
 
 	if (RemoteNodeId.IsValid() && (RemoteNodeId == NodeInfo.NodeId))
@@ -319,7 +316,6 @@ void FUdpMessageProcessor::ProcessByeSegment(FInboundSegment& Segment, FNodeInfo
 void FUdpMessageProcessor::ProcessDataSegment(FInboundSegment& Segment, FNodeInfo& NodeInfo)
 {
 	FUdpMessageSegment::FDataChunk DataChunk;
-
 	*Segment.Data << DataChunk;
 
 	// Discard late segments for sequenced messages
@@ -373,7 +369,6 @@ void FUdpMessageProcessor::ProcessDataSegment(FInboundSegment& Segment, FNodeInf
 void FUdpMessageProcessor::ProcessHelloSegment(FInboundSegment& Segment, FNodeInfo& NodeInfo)
 {
 	FGuid RemoteNodeId;
-
 	*Segment.Data << RemoteNodeId;
 
 	if (RemoteNodeId.IsValid())
@@ -386,7 +381,6 @@ void FUdpMessageProcessor::ProcessHelloSegment(FInboundSegment& Segment, FNodeIn
 void FUdpMessageProcessor::ProcessRetransmitSegment(FInboundSegment& Segment, FNodeInfo& NodeInfo)
 {
 	FUdpMessageSegment::FRetransmitChunk RetransmitChunk;
-
 	*Segment.Data << RetransmitChunk;
 
 	TSharedPtr<FUdpMessageSegmenter> Segmenter = NodeInfo.Segmenters.FindRef(RetransmitChunk.MessageId);
@@ -401,7 +395,6 @@ void FUdpMessageProcessor::ProcessRetransmitSegment(FInboundSegment& Segment, FN
 void FUdpMessageProcessor::ProcessTimeoutSegment(FInboundSegment& Segment, FNodeInfo& NodeInfo)
 {
 	FUdpMessageSegment::FTimeoutChunk TimeoutChunk;
-
 	*Segment.Data << TimeoutChunk;
 
 	TSharedPtr<FUdpMessageSegmenter> Segmenter = NodeInfo.Segmenters.FindRef(TimeoutChunk.MessageId);
@@ -422,7 +415,6 @@ void FUdpMessageProcessor::ProcessUnknownSegment(FInboundSegment& Segment, FNode
 void FUdpMessageProcessor::RemoveKnownNode(const FGuid& NodeId)
 {
 	NodeLostDelegate.ExecuteIfBound(NodeId);
-
 	KnownNodes.Remove(NodeId);
 }
 
@@ -431,7 +423,6 @@ void FUdpMessageProcessor::UpdateKnownNodes()
 {
 	// remove dead remote endpoints
 	FTimespan DeadHelloTimespan = DeadHelloIntervals * Beacon->GetBeaconInterval();
-
 	TArray<FGuid> NodesToRemove;
 
 	for (auto& KnownNodePair : KnownNodes)
@@ -484,22 +475,21 @@ void FUdpMessageProcessor::UpdateSegmenters(FNodeInfo& NodeInfo)
 				DataChunk.MessageId = It.Key();
 				DataChunk.MessageSize = Segmenter->GetMessageSize();
 				DataChunk.SegmentOffset = 1024 * DataChunk.SegmentNumber;
-				DataChunk.Sequence = 0;
+				DataChunk.Sequence = 0; // @todo gmp: implement message sequencing
 				DataChunk.TotalSegments = Segmenter->GetSegmentCount();
 
-				FArrayWriter Writer;
-
-				Writer << Header;
-				Writer << DataChunk;
-
-				if (Sender->Send(MakeShareable(new TArray<uint8>(Writer)), NodeInfo.Endpoint))
+				TSharedRef<FArrayWriter, ESPMode::ThreadSafe> Writer = MakeShareable(new FArrayWriter);
 				{
-					Segmenter->MarkAsSent(DataChunk.SegmentNumber);
+					*Writer << Header;
+					*Writer << DataChunk;
 				}
-				else
+
+				if (!Sender->Send(Writer, NodeInfo.Endpoint))
 				{
 					return;
-				}
+ 				}
+
+				Segmenter->MarkAsSent(DataChunk.SegmentNumber);
 			}
 
 			It.RemoveCurrent();
diff --git a/Engine/Plugins/MovieScene/LevelSequenceEditor/Source/LevelSequenceEditor/LevelSequenceEditor.Build.cs b/Engine/Plugins/MovieScene/LevelSequenceEditor/Source/LevelSequenceEditor/LevelSequenceEditor.Build.cs
--- a/Engine/Plugins/MovieScene/LevelSequenceEditor/Source/LevelSequenceEditor/LevelSequenceEditor.Build.cs
+++ b/Engine/Plugins/MovieScene/LevelSequenceEditor/Source/LevelSequenceEditor/LevelSequenceEditor.Build.cs
@@ -46,6 +46,7 @@ public class LevelSequenceEditor : ModuleRules
 				"LevelSequenceEditor/Private/AssetTools",
                 "LevelSequenceEditor/Private/Customizations",
 				"LevelSequenceEditor/Private/Factories",
+                "LevelSequenceEditor/Private/Misc",
 				"LevelSequenceEditor/Private/Styles",
 			}
         );
diff --git a/Engine/Plugins/MovieScene/LevelSequenceEditor/Source/LevelSequenceEditor/Private/LevelSequenceEditorToolkit.cpp b/Engine/Plugins/MovieScene/LevelSequenceEditor/Source/LevelSequenceEditor/Private/LevelSequenceEditorToolkit.cpp
--- a/Engine/Plugins/MovieScene/LevelSequenceEditor/Source/LevelSequenceEditor/Private/LevelSequenceEditorToolkit.cpp
+++ b/Engine/Plugins/MovieScene/LevelSequenceEditor/Source/LevelSequenceEditor/Private/LevelSequenceEditorToolkit.cpp
@@ -15,198 +15,13 @@
 #include "MovieSceneMaterialTrack.h"
 #include "ScopedTransaction.h"
 #include "ClassIconFinder.h"
-#include "LevelSequenceSpawnRegister.h"
-#include "ISequencerObjectChangeListener.h"
-
 #include "SceneOutlinerModule.h"
 #include "SceneOutlinerPublicTypes.h"
+#include "LevelSequenceEditorSpawnRegister.h"
 
 
 #define LOCTEXT_NAMESPACE "LevelSequenceEditor"
 
-/**
- * Spawn register used in the editor to add some usability features like maintaining selection states, and projecting spawned state onto spawnable defaults
- */
-class FLevelSequenceEditorSpawnRegister : public FLevelSequenceSpawnRegister
-{
-public:
-	/** Weak pointer to the active sequencer */
-	TWeakPtr<ISequencer> WeakSequencer;
-
-	/** Constructor */
-	FLevelSequenceEditorSpawnRegister()
-	{
-		bShouldClearSelectionCache = true;
-
-		FLevelEditorModule& LevelEditor = FModuleManager::GetModuleChecked<FLevelEditorModule>("LevelEditor");
-		OnActorSelectionChangedHandle = LevelEditor.OnActorSelectionChanged().AddRaw(this, &FLevelSequenceEditorSpawnRegister::OnActorSelectionChanged);
-
-		OnActorMovedHandle = GEditor->OnActorMoved().AddLambda([=](AActor* Actor){
-			if (SpawnedObjects.Contains(Actor))
-			{
-				OnSpawnedObjectPropertyChanged(*Actor);
-			}
-		});
-	}
-
-	~FLevelSequenceEditorSpawnRegister()
-	{
-		GEditor->OnActorMoved().Remove(OnActorMovedHandle);
-		if (FLevelEditorModule* LevelEditor = FModuleManager::GetModulePtr<FLevelEditorModule>("LevelEditor"))
-		{
-			LevelEditor->OnActorSelectionChanged().Remove(OnActorSelectionChangedHandle);
-		}
-	}
-
-private:
-
-	/** Called to spawn an object */
-	virtual UObject* SpawnObject(const FGuid& BindingId, FMovieSceneSequenceInstance& SequenceInstance, IMovieScenePlayer& Player) override
-	{
-		TGuardValue<bool> Guard(bShouldClearSelectionCache, false);
-
-		UObject* NewObject = FLevelSequenceSpawnRegister::SpawnObject(BindingId, SequenceInstance, Player);
-		if (NewObject)
-		{
-			SpawnedObjects.Add(NewObject);
-
-			// Add an object listener for the spawned object to propagate changes back onto the spawnable default
-			TSharedPtr<ISequencer> Sequencer = WeakSequencer.Pin();
-			AActor* Actor = Cast<AActor>(NewObject);
-			if (Sequencer.IsValid() && Actor)
-			{
-				Sequencer->GetObjectChangeListener().GetOnAnyPropertyChanged(*NewObject).AddSP(this, &FLevelSequenceEditorSpawnRegister::OnSpawnedObjectPropertyChanged);
-
-				// Select the actor if we think it should be selected
-				if (SelectedSpawnedObjects.Contains(FMovieSceneSpawnRegisterKey(BindingId, SequenceInstance)))
-				{
-					GEditor->SelectActor(Actor, true /*bSelected*/, true /*bNotify*/);
-				}
-			}
-		}
-
-		return NewObject;
-	}
-
-	/** Called right before an object is about to be destroyed */
-	virtual void PreDestroyObject(UObject& Object, const FGuid& BindingId, FMovieSceneSequenceInstance& SequenceInstance) override
-	{
-		TGuardValue<bool> Guard(bShouldClearSelectionCache, false);
-
-		// Cache its selection state
-		AActor* Actor = Cast<AActor>(&Object);
-		if (Actor && GEditor->GetSelectedActors()->IsSelected(Actor))
-		{
-			SelectedSpawnedObjects.Add(FMovieSceneSpawnRegisterKey(BindingId, SequenceInstance));
-			GEditor->SelectActor(Actor, false /*bSelected*/, true /*bNotify*/);
-		}
-
-		SpawnedObjects.Remove(&Object);
-
-		// Remove our object listener
-		TSharedPtr<ISequencer> Sequencer = WeakSequencer.Pin();
-		if (Sequencer.IsValid())
-		{
-			Sequencer->GetObjectChangeListener().ReportObjectDestroyed(Object);
-		}
-
-		FLevelSequenceSpawnRegister::PreDestroyObject(Object, BindingId, SequenceInstance);
-	}
-
-	/** Populate a map of properties that are keyed on a particular object */
-	void PopulateKeyedPropertyMap(AActor& SpawnedObject, TMap<UObject*, TSet<UProperty*>>& OutKeyedPropertyMap)
-	{
-		TSharedPtr<ISequencer> Sequencer = WeakSequencer.Pin();
-
-		Sequencer->GetAllKeyedProperties(SpawnedObject, OutKeyedPropertyMap.FindOrAdd(&SpawnedObject));
-
-		for (UActorComponent* Component : SpawnedObject.GetComponents())
-		{
-			Sequencer->GetAllKeyedProperties(*Component, OutKeyedPropertyMap.FindOrAdd(Component));
-		}
-	}
-
-	/** Called when the editor selection has changed */
-	void OnActorSelectionChanged(const TArray<UObject*>& NewSelection, bool bForceRefresh)
-	{
-		if (bShouldClearSelectionCache)
-		{
-			SelectedSpawnedObjects.Reset();
-		}
-	}
-
-	/** Called when a property on a spawned object changes */
-	void OnSpawnedObjectPropertyChanged(UObject& SpawnedObject)
-	{
-		using namespace EditorUtilities;
-
-		AActor* Actor = CastChecked<AActor>(&SpawnedObject);
-		if (!Actor)
-		{
-			return;
-		}
-
-		TMap<UObject*, TSet<UProperty*>> ObjectToKeyedProperties;
-		PopulateKeyedPropertyMap(*Actor, ObjectToKeyedProperties);
-
-		// Copy any changed actor properties onto the default actor, provided they are not keyed
-		FCopyOptions Options(ECopyOptions::PropagateChangesToArchetypeInstances);
-
-		// Set up a property filter so only stuff that is not keyed gets copied onto the default
-		Options.PropertyFilter = [&](const UProperty& Property, const UObject& Object) -> bool {
-			const TSet<UProperty*>* ExcludedProperties = ObjectToKeyedProperties.Find(const_cast<UObject*>(&Object));
-
-			return !ExcludedProperties || !ExcludedProperties->Contains(const_cast<UProperty*>(&Property));
-		};
-
-		// Now copy the actor properties
-		AActor* DefaultActor = Actor->GetClass()->GetDefaultObject<AActor>();
-		EditorUtilities::CopyActorProperties(Actor, DefaultActor, Options);
-
-		// The above function call explicitly doesn't copy the root component transform (so the default actor is always at 0,0,0)
-		// But in sequencer, we want the object to have a default transform if it doesn't have a transform track
-		static FName RelativeLocation = GET_MEMBER_NAME_CHECKED(USceneComponent, RelativeLocation);
-		static FName RelativeRotation = GET_MEMBER_NAME_CHECKED(USceneComponent, RelativeRotation);
-		static FName RelativeScale3D = GET_MEMBER_NAME_CHECKED(USceneComponent, RelativeScale3D);
-
-		bool bHasKeyedTransform = false;
-		for (UProperty* Property : *ObjectToKeyedProperties.Find(Actor))
-		{
-			FName PropertyName = Property->GetFName();
-			bHasKeyedTransform = PropertyName == RelativeLocation || PropertyName == RelativeRotation || PropertyName == RelativeScale3D;
-			if (bHasKeyedTransform)
-			{
-				break;
-			}
-		}
-
-		// Set the default transform if it's not keyed
-		USceneComponent* RootComponent = Actor->GetRootComponent();
-		USceneComponent* DefaultRootComponent = DefaultActor->GetRootComponent();
-
-		if (!bHasKeyedTransform && RootComponent && DefaultRootComponent)
-		{
-			DefaultRootComponent->RelativeLocation = RootComponent->RelativeLocation;
-			DefaultRootComponent->RelativeRotation = RootComponent->RelativeRotation;
-			DefaultRootComponent->RelativeScale3D = RootComponent->RelativeScale3D;
-		}
-	}
-
-private:
-
-	/** Handles for delegates that we've bound to */
-	FDelegateHandle OnActorMovedHandle, OnActorSelectionChangedHandle;
-
-	/** Set of spawn register keys for objects that should be selected if they are spawned */
-	TSet<FMovieSceneSpawnRegisterKey> SelectedSpawnedObjects;
-
-	/** Set of currently spawned objects */
-	TSet<FObjectKey> SpawnedObjects;
-
-	/** True if we should clear the above selection cache when the editor selection has been changed */
-	bool bShouldClearSelectionCache;
-};
-
 
 /* Local constants
  *****************************************************************************/
@@ -219,6 +34,9 @@ namespace SequencerDefs
 }
 
 
+/* FLevelSequenceEditorToolkit structors
+ *****************************************************************************/
+
 FLevelSequenceEditorToolkit::FLevelSequenceEditorToolkit(const TSharedRef<ISlateStyle>& InStyle)
 	: Style(InStyle)
 {
@@ -253,7 +71,7 @@ FLevelSequenceEditorToolkit::~FLevelSequenceEditorToolkit()
 /* FLevelSequenceEditorToolkit interface
  *****************************************************************************/
 
-void FLevelSequenceEditorToolkit::Initialize( const EToolkitMode::Type Mode, const TSharedPtr<IToolkitHost>& InitToolkitHost, ULevelSequence* InLevelSequence, bool bEditWithinLevelEditor )
+void FLevelSequenceEditorToolkit::Initialize(const EToolkitMode::Type Mode, const TSharedPtr<IToolkitHost>& InitToolkitHost, ULevelSequence* InLevelSequence, bool bEditWithinLevelEditor)
 {
 	// create tab layout
 	const TSharedRef<FTabManager::FLayout> StandaloneDefaultLayout = FTabManager::NewLayout("Standalone_LevelSequenceEditor")
@@ -305,8 +123,7 @@ void FLevelSequenceEditorToolkit::Initialize( const EToolkitMode::Type Mode, con
 	}
 
 	Sequencer = FModuleManager::LoadModuleChecked<ISequencerModule>("Sequencer").CreateSequencer(SequencerInitParams);
-	
-	SpawnRegister->WeakSequencer = Sequencer;
+	SpawnRegister->SetSequencer(Sequencer);
 
 	if (bEditWithinLevelEditor)
 	{
@@ -382,32 +199,32 @@ void FLevelSequenceEditorToolkit::UnregisterTabSpawners(const TSharedRef<class F
 /* FLevelSequenceEditorToolkit callbacks
  *****************************************************************************/
 
-void FLevelSequenceEditorToolkit::HandleAddComponentActionExecute( UActorComponent* Component )
+void FLevelSequenceEditorToolkit::HandleAddComponentActionExecute(UActorComponent* Component)
 {
-	Sequencer->GetHandleToObject( Component );
+	Sequencer->GetHandleToObject(Component);
 }
 
 
-void FLevelSequenceEditorToolkit::HandleAddComponentMaterialActionExecute( UPrimitiveComponent* Component, int32 MaterialIndex )
+void FLevelSequenceEditorToolkit::HandleAddComponentMaterialActionExecute(UPrimitiveComponent* Component, int32 MaterialIndex)
 {
-	FGuid ObjectHandle = Sequencer->GetHandleToObject( Component );
-	for ( const FMovieSceneBinding& Binding : Sequencer->GetFocusedMovieSceneSequence()->GetMovieScene()->GetBindings() )
+	FGuid ObjectHandle = Sequencer->GetHandleToObject(Component);
+	for (const FMovieSceneBinding& Binding : Sequencer->GetFocusedMovieSceneSequence()->GetMovieScene()->GetBindings())
 	{
-		if ( Binding.GetObjectGuid() == ObjectHandle )
+		if (Binding.GetObjectGuid() == ObjectHandle)
 		{
 			bool bHasMaterialTrack = false;
-			for ( UMovieSceneTrack* Track : Binding.GetTracks() )
+			for (UMovieSceneTrack* Track : Binding.GetTracks())
 			{
-				UMovieSceneComponentMaterialTrack* MaterialTrack = Cast<UMovieSceneComponentMaterialTrack>( Track );
-				if ( MaterialTrack != nullptr && MaterialTrack->GetMaterialIndex() == MaterialIndex )
+				UMovieSceneComponentMaterialTrack* MaterialTrack = Cast<UMovieSceneComponentMaterialTrack>(Track);
+				if (MaterialTrack != nullptr && MaterialTrack->GetMaterialIndex() == MaterialIndex)
 				{
 					bHasMaterialTrack = true;
 					break;
 				}
 			}
-			if ( bHasMaterialTrack == false )
+			if (bHasMaterialTrack == false)
 			{
-				const FScopedTransaction Transaction( LOCTEXT( "AddComponentMaterialTrack", "Add component material track" ) );
+				const FScopedTransaction Transaction(LOCTEXT("AddComponentMaterialTrack", "Add component material track"));
 
 				UMovieScene* FocusedMovieScene = Sequencer->GetFocusedMovieSceneSequence()->GetMovieScene();
 				{
@@ -417,7 +234,7 @@ void FLevelSequenceEditorToolkit::HandleAddComponentMaterialActionExecute( UPrim
 				UMovieSceneComponentMaterialTrack* MaterialTrack = Cast<UMovieSceneComponentMaterialTrack>(FocusedMovieScene->AddTrack<UMovieSceneComponentMaterialTrack>(ObjectHandle));
 				{
 					MaterialTrack->Modify();
-					MaterialTrack->SetMaterialIndex( MaterialIndex );
+					MaterialTrack->SetMaterialIndex(MaterialIndex);
 				}
 				
 				Sequencer->NotifyMovieSceneDataChanged();
@@ -507,13 +324,13 @@ void FLevelSequenceEditorToolkit::AddPosessActorMenuExtensions(FMenuBuilder& Men
 		InitOptions.ColumnMap.Add(FBuiltInColumnTypes::Label(), FColumnInfo(EColumnVisibility::Visible, 0));
 
 		// Only display actors that are not possessed already
-		InitOptions.Filters->AddFilterPredicate( FActorFilterPredicate::CreateLambda( IsActorValidForPossession ) );
+		InitOptions.Filters->AddFilterPredicate(FActorFilterPredicate::CreateLambda(IsActorValidForPossession));
 	}
 
 	// actor selector to allow the user to choose an actor
 	FSceneOutlinerModule& SceneOutlinerModule = FModuleManager::LoadModuleChecked<FSceneOutlinerModule>("SceneOutliner");
 	TSharedRef< SWidget > MiniSceneOutliner =
-		SNew( SVerticalBox )
+		SNew(SVerticalBox)
 		+SVerticalBox::Slot()
 		.AutoHeight()
 		.MaxHeight(400.0f)
@@ -538,7 +355,7 @@ void FLevelSequenceEditorToolkit::HandleMapChanged(class UWorld* NewWorld, EMapC
 }
 
 
-TSharedRef<FExtender> FLevelSequenceEditorToolkit::HandleMenuExtensibilityGetExtender( const TSharedRef<FUICommandList> CommandList, const TArray<UObject*> ContextSensitiveObjects )
+TSharedRef<FExtender> FLevelSequenceEditorToolkit::HandleMenuExtensibilityGetExtender(const TSharedRef<FUICommandList> CommandList, const TArray<UObject*> ContextSensitiveObjects)
 {
 	TSharedRef<FExtender> AddTrackMenuExtender(new FExtender());
 	AddTrackMenuExtender->AddMenuExtension(
@@ -570,7 +387,7 @@ TSharedRef<SDockTab> FLevelSequenceEditorToolkit::HandleTabManagerSpawnTab(const
 }
 
 
-void FLevelSequenceEditorToolkit::HandleTrackMenuExtensionAddTrack( FMenuBuilder& AddTrackMenuBuilder, TArray<UObject*> ContextObjects )
+void FLevelSequenceEditorToolkit::HandleTrackMenuExtensionAddTrack(FMenuBuilder& AddTrackMenuBuilder, TArray<UObject*> ContextObjects)
 {
 	if (ContextObjects.Num() != 1)
 	{
@@ -606,8 +423,8 @@ void FLevelSequenceEditorToolkit::HandleTrackMenuExtensionAddTrack( FMenuBuilder
 					{
 						FUIAction AddComponentMaterialAction(FExecuteAction::CreateSP(this, &FLevelSequenceEditorToolkit::HandleAddComponentMaterialActionExecute, Component, MaterialIndex));
 						FText AddComponentMaterialLabel = FText::Format(LOCTEXT("ComponentMaterialIndexLabelFormat", "Element {0}"), FText::AsNumber(MaterialIndex));
-						FText AddComponentMaterialToolTip = FText::Format(LOCTEXT("ComponentMaterialIndexToolTipFormat", "Add material element {0}" ), FText::AsNumber(MaterialIndex));
-						AddTrackMenuBuilder.AddMenuEntry( AddComponentMaterialLabel, AddComponentMaterialToolTip, FSlateIcon(), AddComponentMaterialAction );
+						FText AddComponentMaterialToolTip = FText::Format(LOCTEXT("ComponentMaterialIndexToolTipFormat", "Add material element {0}"), FText::AsNumber(MaterialIndex));
+						AddTrackMenuBuilder.AddMenuEntry(AddComponentMaterialLabel, AddComponentMaterialToolTip, FSlateIcon(), AddComponentMaterialAction);
 					}
 				}
 				AddTrackMenuBuilder.EndSection();
diff --git a/Engine/Plugins/MovieScene/LevelSequenceEditor/Source/LevelSequenceEditor/Private/Misc/LevelSequenceEditorSpawnRegister.cpp b/Engine/Plugins/MovieScene/LevelSequenceEditor/Source/LevelSequenceEditor/Private/Misc/LevelSequenceEditorSpawnRegister.cpp
new file mode 100644
--- /dev/null
+++ b/Engine/Plugins/MovieScene/LevelSequenceEditor/Source/LevelSequenceEditor/Private/Misc/LevelSequenceEditorSpawnRegister.cpp
@@ -0,0 +1,171 @@
+// Copyright 1998-2015 Epic Games, Inc. All Rights Reserved.
+
+#include "LevelSequenceEditorPCH.h"
+#include "ISequencer.h"
+#include "ISequencerObjectChangeListener.h"
+#include "Editor/LevelEditor/Public/LevelEditor.h"
+#include "LevelSequenceEditorSpawnRegister.h"
+
+
+
+/* FLevelSequenceEditorSpawnRegister structors
+ *****************************************************************************/
+
+FLevelSequenceEditorSpawnRegister::FLevelSequenceEditorSpawnRegister()
+{
+	bShouldClearSelectionCache = true;
+
+	FLevelEditorModule& LevelEditor = FModuleManager::GetModuleChecked<FLevelEditorModule>("LevelEditor");
+	OnActorSelectionChangedHandle = LevelEditor.OnActorSelectionChanged().AddRaw(this, &FLevelSequenceEditorSpawnRegister::HandleActorSelectionChanged);
+
+	OnActorMovedHandle = GEditor->OnActorMoved().AddLambda([=](AActor* Actor){
+		HandleAnyPropertyChanged(*Actor);
+	});
+}
+
+
+FLevelSequenceEditorSpawnRegister::~FLevelSequenceEditorSpawnRegister()
+{
+	GEditor->OnActorMoved().Remove(OnActorMovedHandle);
+	if (FLevelEditorModule* LevelEditor = FModuleManager::GetModulePtr<FLevelEditorModule>("LevelEditor"))
+	{
+		LevelEditor->OnActorSelectionChanged().Remove(OnActorSelectionChangedHandle);
+	}
+}
+
+
+/* FLevelSequenceSpawnRegister interface
+ *****************************************************************************/
+
+UObject* FLevelSequenceEditorSpawnRegister::SpawnObject(const FGuid& BindingId, FMovieSceneSequenceInstance& SequenceInstance, IMovieScenePlayer& Player)
+{
+	TGuardValue<bool> Guard(bShouldClearSelectionCache, false);
+	UObject* NewObject = FLevelSequenceSpawnRegister::SpawnObject(BindingId, SequenceInstance, Player);
+
+	// Add an object listener for the spawned object to propagate changes back onto the spawnable default
+	TSharedPtr<ISequencer> Sequencer = WeakSequencer.Pin();
+
+	if (Sequencer.IsValid() && NewObject)
+	{
+		Sequencer->GetObjectChangeListener().GetOnAnyPropertyChanged(*NewObject).AddSP(this, &FLevelSequenceEditorSpawnRegister::HandleAnyPropertyChanged);
+
+		// Select the actor if we think it should be selected
+		AActor* Actor = Cast<AActor>(NewObject);
+		if (Actor && SelectedSpawnedObjects.Contains(FMovieSceneSpawnRegisterKey(BindingId, SequenceInstance)))
+		{
+			GEditor->SelectActor(Actor, true /*bSelected*/, true /*bNotify*/);
+		}
+	}
+
+	return NewObject;
+}
+
+
+void FLevelSequenceEditorSpawnRegister::PreDestroyObject(UObject& Object, const FGuid& BindingId, FMovieSceneSequenceInstance& SequenceInstance)
+{
+	TGuardValue<bool> Guard(bShouldClearSelectionCache, false);
+
+	// Cache its selection state
+	AActor* Actor = Cast<AActor>(&Object);
+	if (Actor && GEditor->GetSelectedActors()->IsSelected(Actor))
+	{
+		SelectedSpawnedObjects.Add(FMovieSceneSpawnRegisterKey(BindingId, SequenceInstance));
+		GEditor->SelectActor(Actor, false /*bSelected*/, true /*bNotify*/);
+	}
+
+	// Remove our object listener
+	TSharedPtr<ISequencer> Sequencer = WeakSequencer.Pin();
+	if (Sequencer.IsValid())
+	{
+		Sequencer->GetObjectChangeListener().ReportObjectDestroyed(Object);
+	}
+
+	FLevelSequenceSpawnRegister::PreDestroyObject(Object, BindingId, SequenceInstance);
+}
+
+
+/* FLevelSequenceEditorSpawnRegister implementation
+ *****************************************************************************/
+
+void FLevelSequenceEditorSpawnRegister::PopulateKeyedPropertyMap(AActor& SpawnedObject, TMap<UObject*, TSet<UProperty*>>& OutKeyedPropertyMap)
+{
+	TSharedPtr<ISequencer> Sequencer = WeakSequencer.Pin();
+	Sequencer->GetAllKeyedProperties(SpawnedObject, OutKeyedPropertyMap.FindOrAdd(&SpawnedObject));
+
+	for (UActorComponent* Component : SpawnedObject.GetComponents())
+	{
+		Sequencer->GetAllKeyedProperties(*Component, OutKeyedPropertyMap.FindOrAdd(Component));
+	}
+}
+
+
+/* FLevelSequenceEditorSpawnRegister callbacks
+ *****************************************************************************/
+
+void FLevelSequenceEditorSpawnRegister::HandleActorSelectionChanged(const TArray<UObject*>& NewSelection, bool bForceRefresh)
+{
+	if (bShouldClearSelectionCache)
+	{
+		SelectedSpawnedObjects.Reset();
+	}
+}
+
+
+void FLevelSequenceEditorSpawnRegister::HandleAnyPropertyChanged(UObject& SpawnedObject)
+{
+	//@todo: AndrewR to filter this to just spawned objects only
+	return;
+
+	using namespace EditorUtilities;
+
+	AActor* Actor = CastChecked<AActor>(&SpawnedObject);
+	if (!Actor)
+	{
+		return;
+	}
+
+	TMap<UObject*, TSet<UProperty*>> ObjectToKeyedProperties;
+	PopulateKeyedPropertyMap(*Actor, ObjectToKeyedProperties);
+
+	// Copy any changed actor properties onto the default actor, provided they are not keyed
+	FCopyOptions Options(ECopyOptions::PropagateChangesToArchetypeInstances);
+
+	// Set up a property filter so only stuff that is not keyed gets copied onto the default
+	Options.PropertyFilter = [&](const UProperty& Property, const UObject& Object) -> bool {
+		const TSet<UProperty*>* ExcludedProperties = ObjectToKeyedProperties.Find(const_cast<UObject*>(&Object));
+
+		return !ExcludedProperties || !ExcludedProperties->Contains(const_cast<UProperty*>(&Property));
+	};
+
+	// Now copy the actor properties
+	AActor* DefaultActor = Actor->GetClass()->GetDefaultObject<AActor>();
+	EditorUtilities::CopyActorProperties(Actor, DefaultActor, Options);
+
+	// The above function call explicitly doesn't copy the root component transform (so the default actor is always at 0,0,0)
+	// But in sequencer, we want the object to have a default transform if it doesn't have a transform track
+	static FName RelativeLocation = GET_MEMBER_NAME_CHECKED(USceneComponent, RelativeLocation);
+	static FName RelativeRotation = GET_MEMBER_NAME_CHECKED(USceneComponent, RelativeRotation);
+	static FName RelativeScale3D = GET_MEMBER_NAME_CHECKED(USceneComponent, RelativeScale3D);
+
+	bool bHasKeyedTransform = false;
+	for (UProperty* Property : *ObjectToKeyedProperties.Find(Actor))
+	{
+		FName PropertyName = Property->GetFName();
+		bHasKeyedTransform = PropertyName == RelativeLocation || PropertyName == RelativeRotation || PropertyName == RelativeScale3D;
+		if (bHasKeyedTransform)
+		{
+			break;
+		}
+	}
+
+	// Set the default transform if it's not keyed
+	USceneComponent* RootComponent = Actor->GetRootComponent();
+	USceneComponent* DefaultRootComponent = DefaultActor->GetRootComponent();
+
+	if (!bHasKeyedTransform && RootComponent && DefaultRootComponent)
+	{
+		DefaultRootComponent->RelativeLocation = RootComponent->RelativeLocation;
+		DefaultRootComponent->RelativeRotation = RootComponent->RelativeRotation;
+		DefaultRootComponent->RelativeScale3D = RootComponent->RelativeScale3D;
+	}
+}
diff --git a/Engine/Plugins/MovieScene/LevelSequenceEditor/Source/LevelSequenceEditor/Private/Misc/LevelSequenceEditorSpawnRegister.h b/Engine/Plugins/MovieScene/LevelSequenceEditor/Source/LevelSequenceEditor/Private/Misc/LevelSequenceEditorSpawnRegister.h
new file mode 100644
--- /dev/null
+++ b/Engine/Plugins/MovieScene/LevelSequenceEditor/Source/LevelSequenceEditor/Private/Misc/LevelSequenceEditorSpawnRegister.h
@@ -0,0 +1,70 @@
+// Copyright 1998-2015 Epic Games, Inc. All Rights Reserved.
+
+#pragma once
+
+#include "LevelSequenceSpawnRegister.h"
+
+
+class AActor;
+class FMovieSceneSequenceInstance;
+class IMovieScenePlayer;
+class ISequencer;
+class UObject;
+class UProperty;
+
+
+/**
+ * Spawn register used in the editor to add some usability features like maintaining selection states, and projecting spawned state onto spawnable defaults
+ */
+class FLevelSequenceEditorSpawnRegister
+	: public FLevelSequenceSpawnRegister
+{
+public:
+
+	/** Constructor */
+	FLevelSequenceEditorSpawnRegister();
+
+	/** Destructor. */
+	~FLevelSequenceEditorSpawnRegister();
+
+public:
+
+	void SetSequencer(const TSharedPtr<ISequencer>& Sequencer)
+	{
+		WeakSequencer = Sequencer;
+	}
+
+public:
+
+	// FLevelSequenceSpawnRegister interface
+
+	virtual UObject* SpawnObject(const FGuid& BindingId, FMovieSceneSequenceInstance& SequenceInstance, IMovieScenePlayer& Player) override;
+	virtual void PreDestroyObject(UObject& Object, const FGuid& BindingId, FMovieSceneSequenceInstance& SequenceInstance) override;
+
+protected:
+
+	/** Populate a map of properties that are keyed on a particular object. */
+	void PopulateKeyedPropertyMap(AActor& SpawnedObject, TMap<UObject*, TSet<UProperty*>>& OutKeyedPropertyMap);
+
+private:
+
+	/** Called when the editor selection has changed. */
+	void HandleActorSelectionChanged(const TArray<UObject*>& NewSelection, bool bForceRefresh);
+
+	/** Called when a property on a spawned object changes. */
+	void HandleAnyPropertyChanged(UObject& SpawnedObject);
+
+private:
+
+	/** Handles for delegates that we've bound to. */
+	FDelegateHandle OnActorMovedHandle, OnActorSelectionChangedHandle;
+
+	/** Set of spawn register keys for objects that should be selected if they are spawned. */
+	TSet<FMovieSceneSpawnRegisterKey> SelectedSpawnedObjects;
+
+	/** True if we should clear the above selection cache when the editor selection has been changed. */
+	bool bShouldClearSelectionCache;
+
+	/** Weak pointer to the active sequencer. */
+	TWeakPtr<ISequencer> WeakSequencer;
+};
diff --git a/Engine/Programs/CrashReportClient/Config/DefaultEngine.ini b/Engine/Programs/CrashReportClient/Config/DefaultEngine.ini
--- a/Engine/Programs/CrashReportClient/Config/DefaultEngine.ini
+++ b/Engine/Programs/CrashReportClient/Config/DefaultEngine.ini
@@ -1,5 +1,6 @@
 [CrashReportClient]
 CrashReportReceiverIP="http://crashreporter.epicgames.com:57005"
+DataRouterUrl="https://datarouter-public-service-prod06.ol.epicgames.com/datarouter/api/v1/public/data"
 bAllowToBeContacted=true
 bSendLogFile=true
 bHideLogFilesOption=true
\ No newline at end of file
diff --git a/Engine/Source/Editor/DetailCustomizations/Private/CaptureTypeCustomization.cpp b/Engine/Source/Editor/DetailCustomizations/Private/CaptureTypeCustomization.cpp
new file mode 100644
--- /dev/null
+++ b/Engine/Source/Editor/DetailCustomizations/Private/CaptureTypeCustomization.cpp
@@ -0,0 +1,82 @@
+﻿// Copyright 1998-2015 Epic Games, Inc. All Rights Reserved.
+
+#include "DetailCustomizationsPrivatePCH.h"
+#include "CaptureTypeCustomization.h"
+#include "MovieSceneCaptureModule.h"
+#include "IPropertyUtilities.h"
+
+#define LOCTEXT_NAMESPACE "CaptureTypeCustomization"
+
+TSharedRef<IPropertyTypeCustomization> FCaptureTypeCustomization::MakeInstance()
+{
+	return MakeShareable( new FCaptureTypeCustomization );
+}
+
+void FCaptureTypeCustomization::CustomizeHeader( TSharedRef<IPropertyHandle> InPropertyHandle, FDetailWidgetRow& HeaderRow, IPropertyTypeCustomizationUtils& CustomizationUtils )
+{
+	PropertyUtilities = CustomizationUtils.GetPropertyUtilities();
+	
+	PropertyHandle = InPropertyHandle->GetChildHandle(GET_MEMBER_NAME_CHECKED(FCaptureProtocolID, Identifier));
+	FMovieSceneCaptureProtocolRegistry& Registry = IMovieSceneCaptureModule::Get().GetProtocolRegistry();
+
+	Registry.IterateProtocols([&](FCaptureProtocolID ID, const FMovieSceneCaptureProtocolInfo& Info){
+		CaptureTypes.Add(MakeShareable(new FCaptureProtocolType{ ID, Info.DisplayName }));
+	});
+
+	if (CaptureTypes.Num() == 0)
+	{
+		return;
+	}
+
+	FName CurrentID;
+	PropertyHandle->GetValue(CurrentID);
+
+	SetCurrentIndex(CurrentID);
+
+	HeaderRow
+	.NameContent()
+	[
+		InPropertyHandle->CreatePropertyNameWidget()
+	]
+	.ValueContent()
+	.MaxDesiredWidth(TOptional<float>())
+	[
+		SNew(SComboBox<TSharedPtr<FCaptureProtocolType>>)
+		.OptionsSource(&CaptureTypes)
+		.OnSelectionChanged(this, &FCaptureTypeCustomization::OnPropertyChanged)
+		.OnGenerateWidget_Lambda([](TSharedPtr<FCaptureProtocolType> CaptureType){
+			return SNew(STextBlock)
+			.Font(IDetailLayoutBuilder::GetDetailFont())
+			.Text(CaptureType->DisplayName);
+		})
+		.InitiallySelectedItem(CaptureTypes[CurrentIndex])
+		[
+			SAssignNew(CurrentText, STextBlock)
+			.Font(IDetailLayoutBuilder::GetDetailFont())
+			.Text(CaptureTypes[CurrentIndex]->DisplayName)
+		]
+	];
+}
+
+void FCaptureTypeCustomization::OnPropertyChanged(TSharedPtr<FCaptureProtocolType> CaptureType, ESelectInfo::Type)
+{
+	SetCurrentIndex(CaptureType->ID.Identifier);
+	PropertyHandle->SetValue(CaptureType->ID.Identifier);
+	CurrentText->SetText(CaptureType->DisplayName);
+
+	PropertyUtilities->RequestRefresh();
+}
+
+void FCaptureTypeCustomization::SetCurrentIndex(FName ID)
+{
+	CurrentIndex = CaptureTypes.IndexOfByPredicate([&](const TSharedPtr<FCaptureProtocolType>& In){
+		return In->ID.Identifier == ID;
+	});
+
+	if (CurrentIndex == INDEX_NONE)
+	{
+		CurrentIndex = 0;
+	}
+}
+
+#undef LOCTEXT_NAMESPACE
diff --git a/Engine/Source/Editor/DetailCustomizations/Private/CaptureTypeCustomization.h b/Engine/Source/Editor/DetailCustomizations/Private/CaptureTypeCustomization.h
new file mode 100644
--- /dev/null
+++ b/Engine/Source/Editor/DetailCustomizations/Private/CaptureTypeCustomization.h
@@ -0,0 +1,39 @@
+// Copyright 1998-2015 Epic Games, Inc. All Rights Reserved.
+
+#pragma once
+
+#include "MovieSceneCaptureProtocolRegistry.h"
+
+class FCaptureTypeCustomization : public IPropertyTypeCustomization
+{
+public:
+	/** Makes a new instance of this detail layout class for a specific detail view requesting it */
+	static TSharedRef<IPropertyTypeCustomization> MakeInstance();
+
+	/** IDetailCustomization interface */
+	virtual void CustomizeHeader( TSharedRef<IPropertyHandle> InPropertyHandle, FDetailWidgetRow& HeaderRow, IPropertyTypeCustomizationUtils& CustomizationUtils ) override;
+	virtual void CustomizeChildren( TSharedRef<IPropertyHandle> InPropertyHandle, IDetailChildrenBuilder& ChildBuilder, IPropertyTypeCustomizationUtils& CustomizationUtils ) override {}
+
+private:
+	struct FCaptureProtocolType
+	{
+		FCaptureProtocolID ID;
+		FText DisplayName;
+	};
+
+	void OnPropertyChanged(TSharedPtr<FCaptureProtocolType> CaptureType, ESelectInfo::Type);
+
+	void SetCurrentIndex(FName ID);
+	
+private:
+	/** Array of available capture types */
+	TArray<TSharedPtr<FCaptureProtocolType>> CaptureTypes;
+	/** The index of the type we're currently displaying */
+	int32 CurrentIndex;
+	/** The text of the current selection */
+	TSharedPtr<STextBlock> CurrentText;
+	/** Property handle of the property we're editing */
+	TSharedPtr<IPropertyHandle> PropertyHandle;
+	/** Property utilites */
+	TSharedPtr<class IPropertyUtilities> PropertyUtilities;
+};
diff --git a/Engine/Source/Editor/DetailCustomizations/Private/DetailCustomizations.cpp b/Engine/Source/Editor/DetailCustomizations/Private/DetailCustomizations.cpp
--- a/Engine/Source/Editor/DetailCustomizations/Private/DetailCustomizations.cpp
+++ b/Engine/Source/Editor/DetailCustomizations/Private/DetailCustomizations.cpp
@@ -101,8 +101,12 @@
 #include "ConfigEditorPropertyDetails.h"
 #include "AssetImportDataCustomization.h"
 #include "CaptureResolutionCustomization.h"
+#include "CaptureTypeCustomization.h"
+#include "RenderPassesCustomization.h"
+#include "MovieSceneCaptureCustomization.h"
 #include "TextCustomization.h"
 #include "AnimTrailNodeDetails.h"
+
 #include "MaterialProxySettingsCustomizations.h"
 #include "ImportantToggleSettingCustomization.h"
 
@@ -211,6 +215,8 @@ void FDetailCustomizationsModule::RegisterPropertyTypeCustomizations()
 	RegisterCustomPropertyTypeLayout("PostProcessSettings", FOnGetPropertyTypeCustomizationInstance::CreateStatic(&FPostProcessSettingsCustomization::MakeInstance));
 	RegisterCustomPropertyTypeLayout("AssetImportInfo", FOnGetPropertyTypeCustomizationInstance::CreateStatic(&FAssetImportDataCustomization::MakeInstance));
 	RegisterCustomPropertyTypeLayout("CaptureResolution", FOnGetPropertyTypeCustomizationInstance::CreateStatic(&FCaptureResolutionCustomization::MakeInstance));
+	RegisterCustomPropertyTypeLayout("CaptureProtocolID", FOnGetPropertyTypeCustomizationInstance::CreateStatic(&FCaptureTypeCustomization::MakeInstance));
+	RegisterCustomPropertyTypeLayout("CompositionGraphCapturePasses", FOnGetPropertyTypeCustomizationInstance::CreateStatic(&FRenderPassesCustomization::MakeInstance));
 	RegisterCustomPropertyTypeLayout("WeightedBlendable", FOnGetPropertyTypeCustomizationInstance::CreateStatic(&FWeightedBlendableCustomization::MakeInstance));
 	RegisterCustomPropertyTypeLayout("MaterialProxySettings", FOnGetPropertyTypeCustomizationInstance::CreateStatic(&FMaterialProxySettingsCustomizations::MakeInstance));
 }
@@ -300,6 +306,8 @@ void FDetailCustomizationsModule::RegisterObjectCustomizations()
 
 	RegisterCustomClassLayout("ConfigHierarchyPropertyView", FOnGetDetailCustomizationInstance::CreateStatic(&FConfigPropertyHelperDetails::MakeInstance));
 
+	RegisterCustomClassLayout("MovieSceneCapture", FOnGetDetailCustomizationInstance::CreateStatic(&FMovieSceneCaptureCustomization::MakeInstance));
+
 	RegisterCustomClassLayout("AnalyticsPrivacySettings", FOnGetDetailCustomizationInstance::CreateStatic(&FImportantToggleSettingCustomization::MakeInstance));
 	RegisterCustomClassLayout("EndUserSettings", FOnGetDetailCustomizationInstance::CreateStatic(&FImportantToggleSettingCustomization::MakeInstance));
 }
diff --git a/Engine/Source/Editor/DetailCustomizations/Private/MovieSceneCaptureCustomization.cpp b/Engine/Source/Editor/DetailCustomizations/Private/MovieSceneCaptureCustomization.cpp
new file mode 100644
--- /dev/null
+++ b/Engine/Source/Editor/DetailCustomizations/Private/MovieSceneCaptureCustomization.cpp
@@ -0,0 +1,44 @@
+// Copyright 1998-2015 Epic Games, Inc. All Rights Reserved.
+
+#include "DetailCustomizationsPrivatePCH.h"
+#include "MovieSceneCaptureCustomization.h"
+#include "MovieSceneCapture.h"
+
+TSharedRef<IDetailCustomization> FMovieSceneCaptureCustomization::MakeInstance()
+{
+	return MakeShareable(new FMovieSceneCaptureCustomization);
+}
+
+void FMovieSceneCaptureCustomization::CustomizeDetails(IDetailLayoutBuilder& DetailBuilder)
+{
+	IDetailLayoutBuilder* CachedBuilder = &DetailBuilder;
+
+	auto UpdateDetails = [=]{
+		CachedBuilder->ForceRefreshDetails();
+	};
+
+	TSharedRef<IPropertyHandle> CaptureTypeProperty = DetailBuilder.GetProperty(GET_MEMBER_NAME_CHECKED(UMovieSceneCapture, CaptureType));
+	CaptureTypeProperty->SetOnPropertyValueChanged(FSimpleDelegate::CreateLambda(UpdateDetails));
+	CaptureTypeProperty->SetOnChildPropertyValueChanged(FSimpleDelegate::CreateLambda(UpdateDetails));
+
+	TSharedRef<IPropertyHandle> ProtocolSettingsProperty = DetailBuilder.GetProperty(GET_MEMBER_NAME_CHECKED(UMovieSceneCapture, ProtocolSettings));
+	DetailBuilder.HideProperty(ProtocolSettingsProperty);
+
+	DetailBuilder.EditCategory("CaptureSettings", FText(), ECategoryPriority::Important);
+
+	UObject* ObjectValue = nullptr;
+	ProtocolSettingsProperty->GetValue(ObjectValue);
+	if (ObjectValue)
+	{
+		IDetailCategoryBuilder& CustomSettingsCategory = DetailBuilder.EditCategory("CustomSettings", ObjectValue->GetClass()->GetDisplayNameText(), ECategoryPriority::TypeSpecific);
+
+		TArray<UObject*> Objects;
+		Objects.Add(ObjectValue);
+
+		for (TFieldIterator<UProperty> Iterator(ObjectValue->GetClass()); Iterator; ++Iterator)
+		{
+			EPropertyLocation::Type Location = Iterator->HasAnyPropertyFlags(CPF_AdvancedDisplay) ? EPropertyLocation::Advanced : EPropertyLocation::Default;
+			CustomSettingsCategory.AddExternalProperty(Objects, Iterator->GetFName(), Location);
+		}
+	}
+}
diff --git a/Engine/Source/Editor/EditorStyle/Private/SlateEditorStyle.cpp b/Engine/Source/Editor/EditorStyle/Private/SlateEditorStyle.cpp
--- a/Engine/Source/Editor/EditorStyle/Private/SlateEditorStyle.cpp
+++ b/Engine/Source/Editor/EditorStyle/Private/SlateEditorStyle.cpp
@@ -1584,6 +1584,20 @@ void FSlateEditorStyle::FStyle::SetupGeneralStyles()
 
 	// Sequencer
 	{
+		Set( "Sequencer.IconKeyAuto", new IMAGE_BRUSH( "Sequencer/IconKeyAuto", Icon12x12 ) );
+		Set( "Sequencer.IconKeyBreak", new IMAGE_BRUSH( "Sequencer/IconKeyBreak", Icon12x12 ) );
+		Set( "Sequencer.IconKeyConstant", new IMAGE_BRUSH( "Sequencer/IconKeyConstant", Icon12x12 ) );
+		Set( "Sequencer.IconKeyLinear", new IMAGE_BRUSH( "Sequencer/IconKeyLinear", Icon12x12 ) );
+		Set( "Sequencer.IconKeyUser", new IMAGE_BRUSH( "Sequencer/IconKeyUser", Icon12x12 ) );
+
+		Set( "Sequencer.KeyCircle", new IMAGE_BRUSH( "Sequencer/KeyCircle", Icon12x12 ) );
+		Set( "Sequencer.KeyDiamond", new IMAGE_BRUSH( "Sequencer/KeyDiamond", Icon12x12 ) );
+		Set( "Sequencer.KeySquare", new IMAGE_BRUSH( "Sequencer/KeySquare", Icon12x12 ) );
+		Set( "Sequencer.KeyTriangle", new IMAGE_BRUSH( "Sequencer/KeyTriangle", Icon12x12 ) );
+		Set( "Sequencer.KeyLeft", new IMAGE_BRUSH( "Common/KeyLeft", FVector2D( 7.0f, 12.0f ) ) );
+		Set( "Sequencer.KeyRight", new IMAGE_BRUSH( "Common/KeyRight", FVector2D( 7.0f, 12.0f ) ) );
+		Set( "Sequencer.PartialKey", new IMAGE_BRUSH( "Sequencer/PartialKey", FVector2D(11.f, 11.f) ) );
+
 		Set( "Sequencer.Timeline.ScrubHandleDown", new BOX_BRUSH( "Sequencer/ScrubHandleDown", FMargin( 6.f/13.f, 5/12.f, 6/13.f, 8/12.f ) ) );
 		Set( "Sequencer.Timeline.ScrubHandleUp", new BOX_BRUSH( "Sequencer/ScrubHandleUp", FMargin( 6.f/13.f, 8/12.f, 6/13.f, 5/12.f ) ) );
 		Set( "Sequencer.Timeline.ScrubHandleWhole", new BOX_BRUSH( "Sequencer/ScrubHandleWhole", FMargin( 6.f/13.f, 10/24.f, 6/13.f, 10/24.f  ) ) );
@@ -1597,33 +1611,28 @@ void FSlateEditorStyle::FStyle::SetupGeneralStyles()
 		Set( "Sequencer.Timeline.PlayRange_R", new BOX_BRUSH( "Sequencer/PlayRange_R", FMargin(0.f, 0.5f, 1.f, 0.5f) ) );
 		Set( "Sequencer.Timeline.PlayRange_Bottom_L", new BOX_BRUSH( "Sequencer/PlayRange_Bottom_L", FMargin(1.f, 0.5f, 0.f, 0.5f) ) );
 		Set( "Sequencer.Timeline.PlayRange_Bottom_R", new BOX_BRUSH( "Sequencer/PlayRange_Bottom_R", FMargin(0.f, 0.5f, 1.f, 0.5f) ) );
+
 		Set( "Sequencer.NotificationImage_AddedPlayMovieSceneEvent", new IMAGE_BRUSH( "Old/Checkbox_checked", Icon16x16 ) );
+
 		Set( "Sequencer.Save", new IMAGE_BRUSH( "Icons/SequencerIcons/icon_Sequencer_Save_24x", Icon40x40 ) );
 		Set( "Sequencer.Save.Small", new IMAGE_BRUSH( "Icons/SequencerIcons/icon_Sequencer_Save_24x", Icon20x20 ) );
-		Set( "Sequencer.IconKeyAuto", new IMAGE_BRUSH( "Sequencer/IconKeyAuto", Icon12x12 ) );
-		Set( "Sequencer.IconKeyBreak", new IMAGE_BRUSH( "Sequencer/IconKeyBreak", Icon12x12 ) );
-		Set( "Sequencer.IconKeyConstant", new IMAGE_BRUSH( "Sequencer/IconKeyConstant", Icon12x12 ) );
-		Set( "Sequencer.IconKeyLinear", new IMAGE_BRUSH( "Sequencer/IconKeyLinear", Icon12x12 ) );
-		Set( "Sequencer.IconKeyUser", new IMAGE_BRUSH( "Sequencer/IconKeyUser", Icon12x12 ) );
-		Set( "Sequencer.KeyCircle", new IMAGE_BRUSH( "Sequencer/KeyCircle", Icon12x12 ) );
-		Set( "Sequencer.KeyDiamond", new IMAGE_BRUSH( "Sequencer/KeyDiamond", Icon12x12 ) );
-		Set( "Sequencer.KeySquare", new IMAGE_BRUSH( "Sequencer/KeySquare", Icon12x12 ) );
-		Set( "Sequencer.KeyTriangle", new IMAGE_BRUSH( "Sequencer/KeyTriangle", Icon12x12 ) );
-		Set( "Sequencer.KeyLeft", new IMAGE_BRUSH( "Common/KeyLeft", FVector2D( 7.0f, 12.0f ) ) );
-		Set( "Sequencer.KeyRight", new IMAGE_BRUSH( "Common/KeyRight", FVector2D( 7.0f, 12.0f ) ) );
-		Set( "Sequencer.PartialKey", new IMAGE_BRUSH( "Sequencer/PartialKey", FVector2D(11.f, 11.f) ) );
 		Set( "Sequencer.GenericSection.Background", new BOX_BRUSH( "Sequencer/GenericSectionBackground", FMargin(5.f/16.f) ) );
 		Set( "Sequencer.DefaultSectionGripLeft", new BOX_BRUSH( "Sequencer/SectionGripLeft", FMargin(5.f/16.f) ) );
 		Set( "Sequencer.DefaultSectionGripRight", new BOX_BRUSH( "Sequencer/SectionGripRight", FMargin(5.f/16.f) ) );
 		Set( "Sequencer.SectionArea.Background", new FSlateColorBrush( FColor::White ) );
-		Set( "Sequencer.AnimationOutliner.SelectionBorder", new BOX_BRUSH( "Common/RoundedSelection_16x", 4.0f/16.0f, SelectionColor_Pressed ) );
-		Set( "Sequencer.AnimationOutliner.SelectionBorderInactive", new BOX_BRUSH("Common/RoundedSelection_16x", 4.0f / 16.0f, SelectionColor_Inactive));
+		Set( "Sequencer.AnimationOutliner.SelectionBorder", new BOX_BRUSH( "Common/Selection_16x", 4.0f/16.0f, SelectionColor_Pressed ) );
+		Set( "Sequencer.AnimationOutliner.SelectionBorderHover", new BOX_BRUSH("Common/Selection_16x", 4.0f / 16.0f, FLinearColor(1.0f, 1.0f, 1.0f, 0.1f) ) );
+		Set( "Sequencer.AnimationOutliner.SelectionBorderInactive", new BOX_BRUSH("Common/Selection_16x", 4.0f / 16.0f, SelectionColor_Inactive));
 		Set( "Sequencer.AnimationOutliner.BoldFont", TTF_CORE_FONT( "Fonts/Roboto-Bold", 11 ) );
 		Set( "Sequencer.AnimationOutliner.RegularFont", TTF_CORE_FONT( "Fonts/Roboto-Regular", 9 ) );
 		Set( "Sequencer.ShotFilter", new IMAGE_BRUSH( "Sequencer/FilteredArea", FVector2D(74,74), FLinearColor::White, ESlateBrushTileType::Both ) );
 		Set( "Sequencer.KeyMark", new IMAGE_BRUSH("Sequencer/KeyMark", FVector2D(3,21), FLinearColor::White, ESlateBrushTileType::NoTile ) );
-		Set( "Sequencer.ToggleAutoKeyEnabled", new IMAGE_BRUSH( "Icons/SequencerIcons/icon_Sequencer_AutoKey_24x", Icon40x40 ) );
-		Set( "Sequencer.ToggleAutoKeyEnabled.Small", new IMAGE_BRUSH( "Icons/SequencerIcons/icon_Sequencer_AutoKey_24x", Icon20x20 ) );
+		Set( "Sequencer.SetAutoKeyModeAll", new IMAGE_BRUSH( "Icons/SequencerIcons/icon_Sequencer_AutoKeyAll_24x", Icon40x40 ) );
+		Set( "Sequencer.SetAutoKeyModeAll.Small", new IMAGE_BRUSH( "Icons/SequencerIcons/icon_Sequencer_AutoKeyAll_24x", Icon20x20 ) );
+		Set( "Sequencer.SetAutoKeyModeAnimated", new IMAGE_BRUSH("Icons/SequencerIcons/icon_Sequencer_AutoKey_24x", Icon40x40 ) );
+		Set( "Sequencer.SetAutoKeyModeAnimated.Small", new IMAGE_BRUSH("Icons/SequencerIcons/icon_Sequencer_AutoKey_24x", Icon20x20 ) );
+		Set( "Sequencer.SetAutoKeyModeNone", new IMAGE_BRUSH("Icons/SequencerIcons/icon_Sequencer_AutoKeyNone_24x", Icon40x40));
+		Set( "Sequencer.SetAutoKeyModeNone.Small", new IMAGE_BRUSH("Icons/SequencerIcons/icon_Sequencer_AutoKeyNone_24x", Icon20x20 ) );
 		Set( "Sequencer.KeyAllEnabled", new IMAGE_BRUSH( "Icons/SequencerIcons/icon_Sequencer_WholeKeyframing_24x", Icon40x40 ) );
 		Set( "Sequencer.KeyAllEnabled.Small", new IMAGE_BRUSH( "Icons/SequencerIcons/icon_Sequencer_WholeKeyframing_24x", Icon20x20 ) );
 		Set( "Sequencer.KeyAllDisabled", new IMAGE_BRUSH( "Icons/SequencerIcons/icon_Sequencer_PartialKeyframing_24x", Icon40x40 ) );
@@ -1642,12 +1651,16 @@ void FSlateEditorStyle::FStyle::SetupGeneralStyles()
 		Set( "Sequencer.Section.SelectionBorder", new BORDER_BRUSH( TEXT("Sequencer/SectionHighlight"), FMargin(7.f/16.f) ) );
 		Set( "Sequencer.Section.LockedBorder", new BORDER_BRUSH( TEXT("Sequencer/SectionLocked"), FMargin(7.f/16.f) ) );
 		Set( "Sequencer.ShotTrack.SectionHandle", new IMAGE_BRUSH( "Old/White", Icon16x16, FLinearColor::Black ) );
+
 		Set( "Sequencer.General", new IMAGE_BRUSH( "Icons/SequencerIcons/icon_Sequencer_View_24x", Icon24x24 ) );
 		Set( "Sequencer.General.Small", new IMAGE_BRUSH( "Icons/SequencerIcons/icon_Sequencer_View_24x", Icon20x20 ) );
 		Set( "Sequencer.Time", new IMAGE_BRUSH( "Icons/SequencerIcons/icon_Sequencer_Time_24x", Icon24x24 ) );
 		Set( "Sequencer.Time.Small", new IMAGE_BRUSH( "Icons/SequencerIcons/icon_Sequencer_Time_24x", Icon20x20 ) );
 		Set( "Sequencer.Value", new IMAGE_BRUSH( "Icons/SequencerIcons/icon_Sequencer_Value_24x", Icon24x24 ) );
 		Set( "Sequencer.Value.Small", new IMAGE_BRUSH( "Icons/SequencerIcons/icon_Sequencer_Value_24x", Icon20x20 ) );
+
+		Set( "Sequencer.TrackArea.LaneColor", FLinearColor(0.3f, 0.3f, 0.3f, 0.3f) );
+
 		Set( "Sequencer.Tracks.Audio", new IMAGE_BRUSH("Icons/AssetIcons/SoundCue_16x", Icon16x16));
 		Set( "Sequencer.Tracks.Event", new IMAGE_BRUSH("Icons/SequencerIcons/icon_Sequencer_Event_16x", Icon16x16));
 		Set( "Sequencer.Tracks.Fade", new IMAGE_BRUSH("Icons/SequencerIcons/icon_Sequencer_Fade_16x", Icon16x16));
diff --git a/Engine/Source/Editor/KismetCompiler/Private/KismetCompilerMisc.cpp b/Engine/Source/Editor/KismetCompiler/Private/KismetCompilerMisc.cpp
--- a/Engine/Source/Editor/KismetCompiler/Private/KismetCompilerMisc.cpp
+++ b/Engine/Source/Editor/KismetCompiler/Private/KismetCompilerMisc.cpp
@@ -189,6 +189,32 @@ bool FKismetCompilerUtilities::IsTypeCompatibleWithProperty(UEdGraphPin* SourceP
 				const UClass* OutputClass = (Direction == EGPD_Output) ? ObjectType : ObjProperty->PropertyClass;
 				const UClass* InputClass = (Direction == EGPD_Output) ? ObjProperty->PropertyClass : ObjectType;
 
+				// Fixup stale types to avoid unwanted mismatches during the reinstancing process
+				if (OutputClass->HasAnyClassFlags(CLASS_NewerVersionExists))
+				{
+					UBlueprint* GeneratedByBP = Cast<UBlueprint>(OutputClass->ClassGeneratedBy);
+					if (GeneratedByBP != nullptr)
+					{
+						TSubclassOf<UObject> NewerClass = GeneratedByBP->GeneratedClass;
+						if (!NewerClass->HasAnyClassFlags(CLASS_NewerVersionExists))
+						{
+							OutputClass = NewerClass;
+						}
+					}
+				}
+				if (InputClass->HasAnyClassFlags(CLASS_NewerVersionExists))
+				{
+					UBlueprint* GeneratedByBP = Cast<UBlueprint>(InputClass->ClassGeneratedBy);
+					if (GeneratedByBP != nullptr)
+					{
+						TSubclassOf<UObject> NewerClass = GeneratedByBP->GeneratedClass;
+						if (!NewerClass->HasAnyClassFlags(CLASS_NewerVersionExists))
+						{
+							InputClass = NewerClass;
+						}
+					}
+				}
+
 				// It matches if it's an exact match or if the output class is more derived than the input class
 				bTypeMismatch = bSubtypeMismatch = !((OutputClass == InputClass) || (OutputClass->IsChildOf(InputClass)));
 
diff --git a/Engine/Source/Editor/MainFrame/Private/Menus/MainMenu.cpp b/Engine/Source/Editor/MainFrame/Private/Menus/MainMenu.cpp
--- a/Engine/Source/Editor/MainFrame/Private/Menus/MainMenu.cpp
+++ b/Engine/Source/Editor/MainFrame/Private/Menus/MainMenu.cpp
@@ -192,7 +192,6 @@ void FMainMenu::FillWindowMenu( FMenuBuilder& MenuBuilder, const TSharedRef< FEx
 	{
 		// This is a temporary home for the spawners of experimental features that must be explicitly enabled.
 		// When the feature becomes permanent and need not check a flag, register a nomad spawner for it in the proper WorkspaceMenu category
-		bool bMessagingDebugger = GetDefault<UEditorExperimentalSettings>()->bMessagingDebugger;
 		bool bBlutility = GetDefault<UEditorExperimentalSettings>()->bEnableEditorUtilityBlueprints;
 		bool bLocalizationDashboard = GetDefault<UEditorExperimentalSettings>()->bEnableLocalizationDashboard;
 		bool bTranslationPicker = GetDefault<UEditorExperimentalSettings>()->bEnableTranslationPicker;
@@ -200,21 +199,10 @@ void FMainMenu::FillWindowMenu( FMenuBuilder& MenuBuilder, const TSharedRef< FEx
 		bool bDeviceOutputLog = GetDefault<UEditorExperimentalSettings>()->bDeviceOutputLog;
 
 		// Make sure at least one is enabled before creating the section
-		if (bMessagingDebugger || bBlutility || bLocalizationDashboard || bTranslationPicker || bMergeActors || bDeviceOutputLog)
+		if (bBlutility || bLocalizationDashboard || bTranslationPicker || bMergeActors || bDeviceOutputLog)
 		{
 			MenuBuilder.BeginSection("ExperimentalTabSpawners", LOCTEXT("ExperimentalTabSpawnersHeading", "Experimental"));
 			{
-				// Messaging Debugger
-				if (bMessagingDebugger)
-				{
-					MenuBuilder.AddMenuEntry(
-						LOCTEXT("MessagingDebuggerLabel", "Messaging Debugger"),
-						LOCTEXT("MessagingDebuggerToolTip", "The Messaging Debugger provides a visual utility for debugging the messaging system."),
-						FSlateIcon(), // Icon lives in the plugin dir for the debugger
-						FUIAction(FExecuteAction::CreateStatic(&FMainMenu::OpenMessagingDebugger))
-						);
-				}
-
 				// Blutility
 				if (bBlutility)
 				{
diff --git a/Engine/Source/Editor/MainFrame/Private/Menus/MainMenu.h b/Engine/Source/Editor/MainFrame/Private/Menus/MainMenu.h
--- a/Engine/Source/Editor/MainFrame/Private/Menus/MainMenu.h
+++ b/Engine/Source/Editor/MainFrame/Private/Menus/MainMenu.h
@@ -75,15 +75,6 @@ private:
 		FGlobalTabmanager::Get()->InvokeTab(FName(TEXT("ProjectLauncher")));
 	}
 
-	/**
-	* Opens the experimental messaging debugger tab.
-	* Remove this when it is no longer experimental.
-	*/
-	static void OpenMessagingDebugger()
-	{
-		FGlobalTabmanager::Get()->InvokeTab(FName(TEXT("MessagingDebugger")));
-	}
-
 	/**
 	* Opens the experimental localization dashboard.
 	* Remove this when it is no longer experimental.
