Copying //UE4/Dev-Sequencer to Dev-Main (//UE4/Dev-Main)

diff --git a/Engine/Plugins/Editor/SequenceRecorder/Source/SequenceRecorder/Private/SequenceRecorderStyle.cpp b/Engine/Plugins/Editor/SequenceRecorder/Source/SequenceRecorder/Private/SequenceRecorderStyle.cpp
new file mode 100644
--- /dev/null
+++ b/Engine/Plugins/Editor/SequenceRecorder/Source/SequenceRecorder/Private/SequenceRecorderStyle.cpp
@@ -0,0 +1,89 @@
+// Copyright 1998-2016 Epic Games, Inc. All Rights Reserved.
+
+#include "SequenceRecorderPrivatePCH.h"
+#include "SequenceRecorderStyle.h"
+#include "SlateStyle.h"
+#include "IPluginManager.h"
+
+#define IMAGE_BRUSH( RelativePath, ... ) FSlateImageBrush( FSequenceRecorderStyle::InContent( RelativePath, ".png" ), __VA_ARGS__ )
+#define BOX_BRUSH( RelativePath, ... ) FSlateBoxBrush( FSequenceRecorderStyle::InContent( RelativePath, ".png" ), __VA_ARGS__ )
+#define BORDER_BRUSH( RelativePath, ... ) FSlateBorderBrush( FSequenceRecorderStyle::InContent( RelativePath, ".png" ), __VA_ARGS__ )
+#define TTF_FONT( RelativePath, ... ) FSlateFontInfo( FSequenceRecorderStyle::InContent( RelativePath, ".ttf" ), __VA_ARGS__ )
+#define TTF_CORE_FONT(RelativePath, ...) FSlateFontInfo( StyleSet->RootToCoreContentDir( RelativePath, TEXT(".ttf") ), __VA_ARGS__ )
+#define OTF_FONT( RelativePath, ... ) FSlateFontInfo( FSequenceRecorderStyle::InContent( RelativePath, ".otf" ), __VA_ARGS__ )
+#define OTF_CORE_FONT(RelativePath, ...) FSlateFontInfo( StyleSet->RootToCoreContentDir( RelativePath, TEXT(".otf") ), __VA_ARGS__ )
+
+FString FSequenceRecorderStyle::InContent( const FString& RelativePath, const ANSICHAR* Extension )
+{
+	static FString ContentDir = IPluginManager::Get().FindPlugin(TEXT("SequenceRecorder"))->GetContentDir();
+	return ( ContentDir / RelativePath ) + Extension;
+}
+
+TSharedPtr< FSlateStyleSet > FSequenceRecorderStyle::StyleSet = NULL;
+TSharedPtr< class ISlateStyle > FSequenceRecorderStyle::Get() { return StyleSet; }
+
+FName FSequenceRecorderStyle::GetStyleSetName()
+{
+	static FName SequenceRecorderStyleName(TEXT("SequenceRecorderStyle"));
+	return SequenceRecorderStyleName;
+}
+
+void FSequenceRecorderStyle::Initialize()
+{
+	// Const icon sizes
+	const FVector2D Icon8x8(8.0f, 8.0f);
+	const FVector2D Icon9x19(9.0f, 19.0f);
+	const FVector2D Icon10x10(10.0f, 10.0f);
+	const FVector2D Icon12x12(12.0f, 12.0f);
+	const FVector2D Icon16x16(16.0f, 16.0f);
+	const FVector2D Icon20x20(20.0f, 20.0f);
+	const FVector2D Icon22x22(22.0f, 22.0f);
+	const FVector2D Icon24x24(24.0f, 24.0f);
+	const FVector2D Icon27x31(27.0f, 31.0f);
+	const FVector2D Icon26x26(26.0f, 26.0f);
+	const FVector2D Icon32x32(32.0f, 32.0f);
+	const FVector2D Icon40x40(40.0f, 40.0f);
+	const FVector2D Icon75x82(75.0f, 82.0f);
+	const FVector2D Icon360x32(360.0f, 32.0f);
+	const FVector2D Icon171x39(171.0f, 39.0f);
+	const FVector2D Icon170x50(170.0f, 50.0f);
+	const FVector2D Icon267x140(170.0f, 50.0f);
+
+	// Only register once
+	if( StyleSet.IsValid() )
+	{
+		return;
+	}
+
+	StyleSet = MakeShareable( new FSlateStyleSet(FSequenceRecorderStyle::GetStyleSetName()) );
+	StyleSet->SetCoreContentRoot(FPaths::EngineContentDir() / TEXT("Slate"));
+
+	// Animation recorder standalone UI
+	{
+		StyleSet->Set( "SequenceRecorder.TabIcon", new IMAGE_BRUSH( "Icons/icon_tab_SequenceRecorder_16x", Icon16x16 ) );
+		StyleSet->Set( "SequenceRecorder.Common.RecordAll", new IMAGE_BRUSH( "Icons/icon_RecordAll_40x", Icon40x40 ) );
+		StyleSet->Set( "SequenceRecorder.Common.StopAll", new IMAGE_BRUSH( "Icons/icon_StopAll_40x", Icon40x40 ) );
+		StyleSet->Set( "SequenceRecorder.Common.AddRecording", new IMAGE_BRUSH( "Icons/icon_AddRecording_40x", Icon40x40 ) );
+		StyleSet->Set( "SequenceRecorder.Common.RemoveRecording", new IMAGE_BRUSH( "Icons/icon_RemoveRecording_40x", Icon40x40 ) );
+	}
+
+	FSlateStyleRegistry::RegisterSlateStyle( *StyleSet.Get() );
+};
+
+#undef IMAGE_BRUSH
+#undef BOX_BRUSH
+#undef BORDER_BRUSH
+#undef TTF_FONT
+#undef TTF_CORE_FONT
+#undef OTF_FONT
+#undef OTF_CORE_FONT
+
+void FSequenceRecorderStyle::Shutdown()
+{
+	if( StyleSet.IsValid() )
+	{
+		FSlateStyleRegistry::UnRegisterSlateStyle( *StyleSet.Get() );
+		ensure( StyleSet.IsUnique() );
+		StyleSet.Reset();
+	}
+}
\ No newline at end of file
diff --git a/Engine/Plugins/Editor/SequenceRecorder/Source/SequenceRecorder/Private/SequenceRecorderStyle.h b/Engine/Plugins/Editor/SequenceRecorder/Source/SequenceRecorder/Private/SequenceRecorderStyle.h
new file mode 100644
--- /dev/null
+++ b/Engine/Plugins/Editor/SequenceRecorder/Source/SequenceRecorder/Private/SequenceRecorderStyle.h
@@ -0,0 +1,23 @@
+// Copyright 1998-2016 Epic Games, Inc. All Rights Reserved.
+
+#pragma once
+
+class FSequenceRecorderStyle
+{
+public:
+
+	static void Initialize();
+	static void Shutdown();
+
+	static TSharedPtr< class ISlateStyle > Get();
+
+	static FName GetStyleSetName();
+
+private:
+
+	static FString InContent( const FString& RelativePath, const ANSICHAR* Extension );
+
+private:
+
+	static TSharedPtr< class FSlateStyleSet > StyleSet;
+};
\ No newline at end of file
diff --git a/Engine/Plugins/Editor/SequenceRecorder/Source/SequenceRecorder/Private/SequenceRecorderUtils.h b/Engine/Plugins/Editor/SequenceRecorder/Source/SequenceRecorder/Private/SequenceRecorderUtils.h
new file mode 100644
--- /dev/null
+++ b/Engine/Plugins/Editor/SequenceRecorder/Source/SequenceRecorder/Private/SequenceRecorderUtils.h
@@ -0,0 +1,40 @@
+// Copyright 1998-2016 Epic Games, Inc. All Rights Reserved.
+
+#pragma once
+
+namespace SequenceRecorderUtils
+{
+
+/**
+ * Utility function that creates an asset with the specified asset path and name.
+ * If the asset cannot be created (as one already exists), we try to postfix the asset
+ * name until we can successfully create the asset.
+ */
+template<typename AssetType>
+static AssetType* MakeNewAsset(const FString& BaseAssetPath, const FString& BaseAssetName)
+{
+	FString AssetPath = BaseAssetPath;
+	FString AssetName = BaseAssetName;
+
+	AssetPath /= AssetName;
+
+	UObject* Package = CreatePackage(nullptr, *AssetPath);
+	UObject* Object = LoadObject<UObject>(Package, *AssetName, nullptr, LOAD_None, nullptr);
+
+	// if object with same name exists, try a different name until we don't find one
+	int32 ExtensionIndex = 0;
+	while(Object)
+	{
+		AssetName = FString::Printf(TEXT("%s_%d"), *BaseAssetName, ExtensionIndex);
+		AssetPath = BaseAssetPath / AssetName;
+		Package = CreatePackage(nullptr, *AssetPath);
+		Object = LoadObject<UObject>(Package, *AssetName, nullptr, LOAD_None, nullptr);
+
+		ExtensionIndex++;
+	}
+
+	// Create the new asset in the package we just made
+	return NewObject<AssetType>(Package, *AssetName, RF_Public | RF_Standalone);
+}
+
+}
\ No newline at end of file
diff --git a/Engine/Plugins/Editor/SequenceRecorder/Source/SequenceRecorder/SequenceRecorder.Build.cs b/Engine/Plugins/Editor/SequenceRecorder/Source/SequenceRecorder/SequenceRecorder.Build.cs
new file mode 100644
--- /dev/null
+++ b/Engine/Plugins/Editor/SequenceRecorder/Source/SequenceRecorder/SequenceRecorder.Build.cs
@@ -0,0 +1,56 @@
+// Copyright 1998-2016 Epic Games, Inc. All Rights Reserved.
+
+namespace UnrealBuildTool.Rules
+{
+	public class SequenceRecorder : ModuleRules
+	{
+		public SequenceRecorder(TargetInfo Target)
+		{
+			PublicIncludePaths.AddRange(
+				new string[] {
+				}
+				);
+
+			PrivateIncludePaths.AddRange(
+				new string[] {
+					"SequenceRecorder/Private",
+					"SequenceRecorder/Private/PropertyRecorders",
+				}
+				);
+
+			PublicDependencyModuleNames.AddRange(
+				new string[]
+				{
+				}
+				);
+
+			PrivateDependencyModuleNames.AddRange(
+				new string[]
+				{
+					"Core",
+					"CoreUObject",
+					"SlateCore",
+					"Slate",
+					"InputCore",
+					"Engine",
+					"UnrealEd",
+					"EditorStyle",
+					"Projects",
+					"LevelEditor",
+					"Persona",
+					"WorkspaceMenuStructure",
+					"PropertyEditor",
+					"MovieScene",
+					"MovieSceneTracks",
+					"LevelSequence",
+				}
+				);
+
+			DynamicallyLoadedModuleNames.AddRange(
+				new string[]
+				{
+				}
+				);
+		}
+	}
+}
\ No newline at end of file
diff --git a/Engine/Plugins/Experimental/SignificanceManager/Source/SignificanceManager/Private/SignificanceManager.cpp b/Engine/Plugins/Experimental/SignificanceManager/Source/SignificanceManager/Private/SignificanceManager.cpp
--- a/Engine/Plugins/Experimental/SignificanceManager/Source/SignificanceManager/Private/SignificanceManager.cpp
+++ b/Engine/Plugins/Experimental/SignificanceManager/Source/SignificanceManager/Private/SignificanceManager.cpp
@@ -225,6 +225,16 @@ const TArray<const USignificanceManager::FManagedObjectInfo*>& USignificanceMana
 	return EmptySet;
 }
 
+USignificanceManager::FManagedObjectInfo* USignificanceManager::GetManagedObject(UObject* Object) const
+{
+	if (FManagedObjectInfo* const* Info = ManagedObjects.Find(Object))
+	{
+		return (*Info);
+	}
+
+	return nullptr;
+}
+
 void USignificanceManager::GetManagedObjects(TArray<const USignificanceManager::FManagedObjectInfo*>& OutManagedObjects, bool bInSignificanceOrder) const
 {
 	OutManagedObjects.Reserve(ManagedObjects.Num());
diff --git a/Engine/Plugins/Experimental/SignificanceManager/Source/SignificanceManager/Public/SignificanceManager.h b/Engine/Plugins/Experimental/SignificanceManager/Source/SignificanceManager/Public/SignificanceManager.h
--- a/Engine/Plugins/Experimental/SignificanceManager/Source/SignificanceManager/Public/SignificanceManager.h
+++ b/Engine/Plugins/Experimental/SignificanceManager/Source/SignificanceManager/Public/SignificanceManager.h
@@ -86,6 +86,8 @@ public:
 		const UObject* GetObject() const { return Object; }
 		FName GetTag() const { return Tag; }
 		float GetSignificance() const { return Significance; }
+		FSignificanceFunction GetSignificanceFunction() const { return SignificanceFunction; }
+		FSignificanceNotify GetSignificanceNotifyDelegate() const { return SignificanceNotifyDelegate; }	
 
 	private:
 		const UObject* Object;
@@ -127,6 +129,9 @@ public:
 	// Returns all managed objects regardless of tag
 	void GetManagedObjects(TArray<const FManagedObjectInfo*>& OutManagedObjects, bool bInSignificanceOrder = false) const;
 
+	// Returns the managed object for the passed-in object, if any. Otherwise returns nullptr
+	USignificanceManager::FManagedObjectInfo* GetManagedObject(UObject* Object) const;
+
 	// Returns the significance value for a given object, returns 0 if object is not managed
 	float GetSignificance(const UObject* Object) const;
 
diff --git a/Engine/Plugins/Experimental/StereoPanorama/Source/StereoPanorama/Private/SceneCapturer.cpp b/Engine/Plugins/Experimental/StereoPanorama/Source/StereoPanorama/Private/SceneCapturer.cpp
--- a/Engine/Plugins/Experimental/StereoPanorama/Source/StereoPanorama/Private/SceneCapturer.cpp
+++ b/Engine/Plugins/Experimental/StereoPanorama/Source/StereoPanorama/Private/SceneCapturer.cpp
@@ -62,7 +62,7 @@ void USceneCapturer::InitCaptureComponent( USceneCaptureComponent2D* CaptureComp
     CaptureComponent->bCaptureEveryFrame = false;
     CaptureComponent->CaptureSource = ESceneCaptureSource::SCS_FinalColorLDR;
 
-    CaptureComponent->TextureTarget = NewObject<UTextureRenderTarget2D>(this);
+    CaptureComponent->TextureTarget = NewObject<UTextureRenderTarget2D>(this, TEXT("SceneCaptureTextureTarget"));
     //TODO: ikrimae: Not sure why the render target needs to be float to avoid banding. Seems like captures to this RT and then applies PP
     //               on top of it which causes degredation.
     CaptureComponent->TextureTarget->InitCustomFormat(CaptureWidth, CaptureHeight, PF_A16B16G16R16, false);
diff --git a/Engine/Plugins/Experimental/StereoPanorama/Source/StereoPanorama/Private/SceneCapturer.h b/Engine/Plugins/Experimental/StereoPanorama/Source/StereoPanorama/Private/SceneCapturer.h
--- a/Engine/Plugins/Experimental/StereoPanorama/Source/StereoPanorama/Private/SceneCapturer.h
+++ b/Engine/Plugins/Experimental/StereoPanorama/Source/StereoPanorama/Private/SceneCapturer.h
@@ -4,6 +4,7 @@
 
 #include "SceneCapturer.generated.h"
 
+
 DECLARE_LOG_CATEGORY_EXTERN( LogStereoPanorama, Log, All );
 
 DECLARE_STATS_GROUP( TEXT( "SP" ), STATGROUP_SP, STATCAT_Advanced );
@@ -13,22 +14,22 @@ DECLARE_CYCLE_STAT( TEXT( "SampleSpherical" ), STAT_SPSampleSpherical, STATGROUP
 DECLARE_CYCLE_STAT( TEXT( "ReadStrip" ),       STAT_SPReadStrip,       STATGROUP_SP );
 DECLARE_CYCLE_STAT( TEXT( "FillAlpha" ),       STAT_SPFillAlpha,       STATGROUP_SP );
 
+
 UENUM()
-namespace ECaptureStep
+enum class ECaptureStep
 {
-	enum Type
-	{
-		Reset,
-        SetStartPosition,
-		SetPosition,
-		Read,
-		Pause,
-		Unpause,
-	};
-}
+	Reset,
+    SetStartPosition,
+	SetPosition,
+	Read,
+	Pause,
+	Unpause
+};
+
 
 DECLARE_DELEGATE_TwoParams(FStereoCaptureDoneDelegate, const TArray<FColor>&, const TArray<FColor>&);
 
+
 UCLASS()
 class USceneCapturer 
 	: public UObject
@@ -37,12 +38,17 @@ class USceneCapturer
     GENERATED_BODY()
 
 public:
+
     USceneCapturer();
 
     //NOTE: ikrimae: Adding this ctor hack to fix the 4.8p2 problem with hot reload macros calling empty constructors
     //               https://answers.unrealengine.com/questions/228042/48p2-compile-fails-on-class-with-non-default-const.html
     USceneCapturer(FVTableHelper& Helper);
     
+public:
+
+	//~ FTickableGameObject interface
+
 	virtual void Tick( float DeltaTime ) override;
 
 	virtual bool IsTickable() const 
@@ -60,6 +66,8 @@ public:
 		RETURN_QUICK_DECLARE_CYCLE_STAT( USceneCapturer, STATGROUP_Tickables );
 	}
 
+public:
+
 	void InitCaptureComponent( USceneCaptureComponent2D* CaptureComponent, float HFov, float VFov, EStereoscopicPass InStereoPass );
 
 	void CaptureComponent( int32 CurrentHorizontalStep, int32 CurrentVerticalStep, FString Folder, USceneCaptureComponent2D* CaptureComponent, TArray<FColor>& Atlas );
@@ -76,6 +84,8 @@ public:
 
 	void Reset();
 
+public:
+
 	IImageWrapperModule& ImageWrapperModule;
 
 	bool bIsTicking;
@@ -88,7 +98,7 @@ public:
     int32 StartFrame;
     int32 EndFrame;
 
-	ECaptureStep::Type CaptureStep;
+	ECaptureStep CaptureStep;
 	int32 CurrentFrameCount;
 
 	int32 CaptureWidth;
@@ -116,6 +126,7 @@ public:
 	}
 
 private:
+
     const float hAngIncrement;
     const float vAngIncrement;
     const float eyeSeparation;
@@ -151,4 +162,4 @@ private:
 	FString FrameDescriptors;
 
     FStereoCaptureDoneDelegate StereoCaptureDoneDelegate;
-};
\ No newline at end of file
+};
diff --git a/Engine/Plugins/Experimental/StereoPanorama/Source/StereoPanorama/Private/StereoPanorama.cpp b/Engine/Plugins/Experimental/StereoPanorama/Source/StereoPanorama/Private/StereoPanorama.cpp
--- a/Engine/Plugins/Experimental/StereoPanorama/Source/StereoPanorama/Private/StereoPanorama.cpp
+++ b/Engine/Plugins/Experimental/StereoPanorama/Source/StereoPanorama/Private/StereoPanorama.cpp
@@ -2,27 +2,30 @@
 
 #include "StereoPanoramaPrivatePCH.h"
 
+
 TSharedPtr<FStereoPanoramaManager> StereoPanoramaManager;
 
-/** IModuleInterface - initialize the module */
+
 void FStereoPanoramaModule::StartupModule()
 {
-	StereoPanoramaManager = MakeShareable( new FStereoPanoramaManager() );
+	StereoPanoramaManager = MakeShareable(new FStereoPanoramaManager());
 }
 
-/** IModuleInterface - shutdown the module */
+
 void FStereoPanoramaModule::ShutdownModule()
 {
-	if( StereoPanoramaManager.IsValid() )
+	if (StereoPanoramaManager.IsValid())
 	{
 		StereoPanoramaManager.Reset();
 	}
 }
 
+
 TSharedPtr<FStereoPanoramaManager> FStereoPanoramaModule::Get()
 {
-	check( StereoPanoramaManager.IsValid() );
+	check(StereoPanoramaManager.IsValid());
 	return StereoPanoramaManager;
 }
 
-IMPLEMENT_MODULE( FStereoPanoramaModule, StereoPanorama )
+
+IMPLEMENT_MODULE(FStereoPanoramaModule, StereoPanorama)
diff --git a/Engine/Plugins/Experimental/StereoPanorama/Source/StereoPanorama/Private/StereoPanoramaManager.cpp b/Engine/Plugins/Experimental/StereoPanorama/Source/StereoPanorama/Private/StereoPanoramaManager.cpp
--- a/Engine/Plugins/Experimental/StereoPanorama/Source/StereoPanorama/Private/StereoPanoramaManager.cpp
+++ b/Engine/Plugins/Experimental/StereoPanorama/Source/StereoPanorama/Private/StereoPanoramaManager.cpp
@@ -4,36 +4,37 @@
 
 
 //Slice Controls
-IConsoleVariable* FStereoPanoramaManager::HorizontalAngularIncrement = IConsoleManager::Get().RegisterConsoleVariable( TEXT( "SP.HorizontalAngularIncrement" ), 1.0f, TEXT( "The number of degrees per horizontal step. Must be a factor of 360." ), ECVF_Default );
-IConsoleVariable* FStereoPanoramaManager::VerticalAngularIncrement   = IConsoleManager::Get().RegisterConsoleVariable( TEXT( "SP.VerticalAngularIncrement" ), 90.0f, TEXT( "The number of degrees per vertical step. Must be a factor of 180." ), ECVF_Default );
-IConsoleVariable* FStereoPanoramaManager::CaptureHorizontalFOV       = IConsoleManager::Get().RegisterConsoleVariable( TEXT( "SP.CaptureHorizontalFOV"), 90.0f, TEXT("Horizontal FOV for scene capture component. Must be larger than SP.HorizontalAngularIncrement"), ECVF_Default );
+IConsoleVariable* FStereoPanoramaManager::HorizontalAngularIncrement = IConsoleManager::Get().RegisterConsoleVariable(TEXT("SP.HorizontalAngularIncrement"), 1.0f, TEXT("The number of degrees per horizontal step. Must be a factor of 360."), ECVF_Default);
+IConsoleVariable* FStereoPanoramaManager::VerticalAngularIncrement   = IConsoleManager::Get().RegisterConsoleVariable(TEXT("SP.VerticalAngularIncrement"), 90.0f, TEXT("The number of degrees per vertical step. Must be a factor of 180."), ECVF_Default);
+IConsoleVariable* FStereoPanoramaManager::CaptureHorizontalFOV       = IConsoleManager::Get().RegisterConsoleVariable(TEXT("SP.CaptureHorizontalFOV"), 90.0f, TEXT("Horizontal FOV for scene capture component. Must be larger than SP.HorizontalAngularIncrement"), ECVF_Default);
 
 //Atlas Controls
-IConsoleVariable* FStereoPanoramaManager::StepCaptureWidth           = IConsoleManager::Get().RegisterConsoleVariable( TEXT( "SP.StepCaptureWidth" ), 4096, TEXT( "The final spherical atlas width" ), ECVF_Default );
-IConsoleVariable* FStereoPanoramaManager::EyeSeparation              = IConsoleManager::Get().RegisterConsoleVariable( TEXT( "SP.EyeSeparation" ), 6.4f, TEXT( "The separation of the stereo cameras" ), ECVF_Default );
-IConsoleVariable* FStereoPanoramaManager::ForceAlpha                 = IConsoleManager::Get().RegisterConsoleVariable( TEXT( "SP.ForceAlpha" ), false, TEXT( "Force the alpha value to completely opaque" ), ECVF_Default );
+IConsoleVariable* FStereoPanoramaManager::StepCaptureWidth           = IConsoleManager::Get().RegisterConsoleVariable(TEXT("SP.StepCaptureWidth"), 4096, TEXT("The final spherical atlas width"), ECVF_Default);
+IConsoleVariable* FStereoPanoramaManager::EyeSeparation              = IConsoleManager::Get().RegisterConsoleVariable(TEXT("SP.EyeSeparation"), 6.4f, TEXT("The separation of the stereo cameras"), ECVF_Default);
+IConsoleVariable* FStereoPanoramaManager::ForceAlpha                 = IConsoleManager::Get().RegisterConsoleVariable(TEXT("SP.ForceAlpha"), false, TEXT("Force the alpha value to completely opaque"), ECVF_Default);
 
 //Sampling Controls
-IConsoleVariable* FStereoPanoramaManager::CaptureSlicePixelWidth     = IConsoleManager::Get().RegisterConsoleVariable( TEXT( "SP.CaptureSlicePixelWidth"), 2048, TEXT(" Capture Slice Pixel Dimension"), ECVF_Default);
-IConsoleVariable* FStereoPanoramaManager::EnableBilerp               = IConsoleManager::Get().RegisterConsoleVariable( TEXT( "SP.EnableBilerp"), true, TEXT("0 - No Filtering 1- Bilinear Filter slice samples"), ECVF_Default);
-IConsoleVariable* FStereoPanoramaManager::SuperSamplingMethod        = IConsoleManager::Get().RegisterConsoleVariable( TEXT( "SP.SuperSamplingMethod"), 1, TEXT(" 0 - No Supersampling, 1 - Rotated Grid SS"), ECVF_Default);
+IConsoleVariable* FStereoPanoramaManager::CaptureSlicePixelWidth     = IConsoleManager::Get().RegisterConsoleVariable(TEXT("SP.CaptureSlicePixelWidth"), 2048, TEXT(" Capture Slice Pixel Dimension"), ECVF_Default);
+IConsoleVariable* FStereoPanoramaManager::EnableBilerp               = IConsoleManager::Get().RegisterConsoleVariable(TEXT("SP.EnableBilerp"), true, TEXT("0 - No Filtering 1- Bilinear Filter slice samples"), ECVF_Default);
+IConsoleVariable* FStereoPanoramaManager::SuperSamplingMethod        = IConsoleManager::Get().RegisterConsoleVariable(TEXT("SP.SuperSamplingMethod"), 1, TEXT(" 0 - No Supersampling, 1 - Rotated Grid SS"), ECVF_Default);
 
 //Debug Controls
-//IConsoleVariable* FStereoPanoramaManager::ConcurrentCaptures         = IConsoleManager::Get().RegisterConsoleVariable( TEXT( "SP.ConcurrentCaptures" ), 15, TEXT("The number of scene captures to capture at the same time"), ECVF_Default);
-IConsoleVariable* FStereoPanoramaManager::ConcurrentCaptures         = IConsoleManager::Get().RegisterConsoleVariable( TEXT("SP.ConcurrentCaptures"), 30, TEXT("The number of scene captures to capture at the same time"), ECVF_Default);
-IConsoleVariable* FStereoPanoramaManager::GenerateDebugImages        = IConsoleManager::Get().RegisterConsoleVariable( TEXT( "SP.GenerateDebugImages" ), 0, TEXT("0 - No Debug Images\n1 - Save out each strip as it is generated\n2 - Save each entire slice"), ECVF_Default);
-IConsoleVariable* FStereoPanoramaManager::OutputDir                  = IConsoleManager::Get().RegisterConsoleVariable( TEXT( "SP.OutputDir" ), TEXT("F:/NeuroSP"), TEXT("Output directory"), ECVF_Default);
-IConsoleVariable* FStereoPanoramaManager::ShouldOverrideInitialYaw   = IConsoleManager::Get().RegisterConsoleVariable( TEXT( "SP.ShouldOverrideInitialYaw" ), true, TEXT("Override Initial Camera Yaw. Set to true if you don't want to use PlayerController View Dir"), ECVF_Default);
-IConsoleVariable* FStereoPanoramaManager::ForcedInitialYaw           = IConsoleManager::Get().RegisterConsoleVariable( TEXT( "SP.ForcedInitialYaw" ), 90.0f, TEXT("Yaw value for initial Camera view direction. Set ShouldOverrideInitialYaw to true to use this value"), ECVF_Default);
-IConsoleVariable* FStereoPanoramaManager::FadeStereoToZeroAtSides    = IConsoleManager::Get().RegisterConsoleVariable( TEXT( "SP.FadeStereoToZeroAtSides"), false, TEXT("Fade stereo effect between left/right eye to zero at 90 degrees."), ECVF_Default);
+//IConsoleVariable* FStereoPanoramaManager::ConcurrentCaptures         = IConsoleManager::Get().RegisterConsoleVariable(TEXT("SP.ConcurrentCaptures"), 15, TEXT("The number of scene captures to capture at the same time"), ECVF_Default);
+IConsoleVariable* FStereoPanoramaManager::ConcurrentCaptures         = IConsoleManager::Get().RegisterConsoleVariable(TEXT("SP.ConcurrentCaptures"), 30, TEXT("The number of scene captures to capture at the same time"), ECVF_Default);
+IConsoleVariable* FStereoPanoramaManager::GenerateDebugImages        = IConsoleManager::Get().RegisterConsoleVariable(TEXT("SP.GenerateDebugImages"), 0, TEXT("0 - No Debug Images\n1 - Save out each strip as it is generated\n2 - Save each entire slice"), ECVF_Default);
+IConsoleVariable* FStereoPanoramaManager::OutputDir                  = IConsoleManager::Get().RegisterConsoleVariable(TEXT("SP.OutputDir"), TEXT(""), TEXT("Output directory"), ECVF_Default);
+IConsoleVariable* FStereoPanoramaManager::ShouldOverrideInitialYaw   = IConsoleManager::Get().RegisterConsoleVariable(TEXT("SP.ShouldOverrideInitialYaw"), true, TEXT("Override Initial Camera Yaw. Set to true if you don't want to use PlayerController View Dir"), ECVF_Default);
+IConsoleVariable* FStereoPanoramaManager::ForcedInitialYaw           = IConsoleManager::Get().RegisterConsoleVariable(TEXT("SP.ForcedInitialYaw"), 90.0f, TEXT("Yaw value for initial Camera view direction. Set ShouldOverrideInitialYaw to true to use this value"), ECVF_Default);
+IConsoleVariable* FStereoPanoramaManager::FadeStereoToZeroAtSides    = IConsoleManager::Get().RegisterConsoleVariable(TEXT("SP.FadeStereoToZeroAtSides"), false, TEXT("Fade stereo effect between left/right eye to zero at 90 degrees."), ECVF_Default);
 
 
 void FStereoPanoramaManager::PanoramicScreenshot(const TArray<FString>& Args)
 {
-    FStereoCaptureDoneDelegate emptyDelegate;
-    FStereoPanoramaManager::PanoramicScreenshot(0, 0, emptyDelegate);
+    FStereoCaptureDoneDelegate EmptyDelegate;
+    FStereoPanoramaManager::PanoramicScreenshot(0, 0, EmptyDelegate);
 }
 
+
 void FStereoPanoramaManager::PanoramicScreenshot(const int32 InStartFrame, const int32 InEndFrame, FStereoCaptureDoneDelegate& InStereoCaptureDoneDelegate)
 {
 	// Construct a capturer that has stereo USceneCaptureComponent2D components
@@ -41,83 +42,87 @@ void FStereoPanoramaManager::PanoramicScreenshot(const int32 InStartFrame, const
 	SceneCapturer->AddToRoot();
 
 	// Rotation is ignored; always start from a yaw of zero
-	SceneCapturer->SetInitialState( InStartFrame, InEndFrame, InStereoCaptureDoneDelegate );
+	SceneCapturer->SetInitialState(InStartFrame, InEndFrame, InStereoCaptureDoneDelegate);
 }
 
+
 void FStereoPanoramaManager::Cleanup()
 {
-	if( SceneCapturer != NULL )
+	if (SceneCapturer != nullptr)
 	{
 		SceneCapturer->Reset();
 		SceneCapturer->RemoveFromRoot();
 
 		// Let GC handle the deletion
-		SceneCapturer = NULL;
+		SceneCapturer = nullptr;
 	}
 }
 
-void FStereoPanoramaManager::PanoramicMovie( const TArray<FString>& Args )
+
+void FStereoPanoramaManager::PanoramicMovie(const TArray<FString>& Args)
 {
     int32 StartFrame = 0;
     int32 EndFrame   = 0;
 
-	if( Args.Num() == 1 )
+	if (Args.Num() == 1)
 	{
         StartFrame = 0;
-        EndFrame   = FCString::Atoi( *Args[0] ) - 1;    //Frame Range is inclusive so -1 to find the last frame
+        EndFrame   = FCString::Atoi(*Args[0]) - 1;    //Frame Range is inclusive so -1 to find the last frame
 	}
     else if (Args.Num() == 2)
     {
-        StartFrame = FCString::Atoi( *Args[0] );
-        EndFrame   = FCString::Atoi( *Args[1] );
-        EndFrame   = FMath::Max(StartFrame, EndFrame);
+        StartFrame = FCString::Atoi(*Args[0]);
+        EndFrame = FCString::Atoi(*Args[1]);
+        EndFrame = FMath::Max(StartFrame, EndFrame);
     }
 
-    FStereoCaptureDoneDelegate emptyDelegate;
-    FStereoPanoramaManager::PanoramicScreenshot(StartFrame, EndFrame, emptyDelegate);
+    FStereoCaptureDoneDelegate EmptyDelegate;
+    FStereoPanoramaManager::PanoramicScreenshot(StartFrame, EndFrame, EmptyDelegate);
 }
 
-void FStereoPanoramaManager::PanoramicQuality( const TArray<FString>& Args )
+
+void FStereoPanoramaManager::PanoramicQuality(const TArray<FString>& Args)
 {
-	if( Args.Contains( TEXT( "preview" ) ) )
+	if (Args.Contains(TEXT("preview")))
 	{
-		UE_LOG( LogStereoPanorama, Display, TEXT( " ... setting 'preview' quality" ) );
+		UE_LOG(LogStereoPanorama, Display, TEXT(" ... setting 'preview' quality"));
 
-		FStereoPanoramaManager::HorizontalAngularIncrement->Set( TEXT( "5" ) );
-		FStereoPanoramaManager::VerticalAngularIncrement->Set( TEXT( "60" ) );
-		FStereoPanoramaManager::CaptureHorizontalFOV->Set( TEXT( "60" ) );
-		FStereoPanoramaManager::StepCaptureWidth->Set( TEXT( "720" ) );
+		FStereoPanoramaManager::HorizontalAngularIncrement->Set(TEXT("5"));
+		FStereoPanoramaManager::VerticalAngularIncrement->Set(TEXT("60"));
+		FStereoPanoramaManager::CaptureHorizontalFOV->Set(TEXT("60"));
+		FStereoPanoramaManager::StepCaptureWidth->Set(TEXT("720"));
 	}
-	else if( Args.Contains( TEXT( "average" ) ) )
+	else if (Args.Contains(TEXT("average")))
 	{
-		UE_LOG( LogStereoPanorama, Display, TEXT( " ... setting 'average' quality" ) );
+		UE_LOG(LogStereoPanorama, Display, TEXT(" ... setting 'average' quality"));
 		
-		FStereoPanoramaManager::HorizontalAngularIncrement->Set( TEXT( "2" ) );
-		FStereoPanoramaManager::VerticalAngularIncrement->Set( TEXT( "30" ) );
-		FStereoPanoramaManager::CaptureHorizontalFOV->Set( TEXT( "30" ) );
-		FStereoPanoramaManager::StepCaptureWidth->Set( TEXT( "1440" ) );
+		FStereoPanoramaManager::HorizontalAngularIncrement->Set(TEXT("2"));
+		FStereoPanoramaManager::VerticalAngularIncrement->Set(TEXT("30"));
+		FStereoPanoramaManager::CaptureHorizontalFOV->Set(TEXT("30"));
+		FStereoPanoramaManager::StepCaptureWidth->Set(TEXT("1440"));
 	}
-	else if( Args.Contains( TEXT( "improved" ) ) )
+	else if (Args.Contains(TEXT("improved")))
 	{
-		UE_LOG( LogStereoPanorama, Display, TEXT( " ... setting 'improved' quality" ) );
+		UE_LOG(LogStereoPanorama, Display, TEXT(" ... setting 'improved' quality"));
 		
-		FStereoPanoramaManager::HorizontalAngularIncrement->Set( TEXT( "0.5" ) );
-		FStereoPanoramaManager::VerticalAngularIncrement->Set( TEXT( "22.5" ) );
-		FStereoPanoramaManager::CaptureHorizontalFOV->Set( TEXT( "22.5" ) );
-		FStereoPanoramaManager::StepCaptureWidth->Set( TEXT( "1440" ) );
+		FStereoPanoramaManager::HorizontalAngularIncrement->Set(TEXT("0.5"));
+		FStereoPanoramaManager::VerticalAngularIncrement->Set(TEXT("22.5"));
+		FStereoPanoramaManager::CaptureHorizontalFOV->Set(TEXT("22.5"));
+		FStereoPanoramaManager::StepCaptureWidth->Set(TEXT("1440"));
 	}
 	else
 	{
-		UE_LOG( LogStereoPanorama, Warning, TEXT( "No quality setting found; options are 'preview | average | improved'" ) );
+		UE_LOG(LogStereoPanorama, Warning, TEXT("No quality setting found; options are 'preview | average | improved'"));
 	}
 }
 
+
 void FStereoPanoramaManager::PanoramicTogglePause(const TArray<FString>& Args)
 {
     auto CapturePlayerController = UGameplayStatics::GetPlayerController(GWorld, 0);
     auto CaptureGameMode = UGameplayStatics::GetGameMode(GWorld);
 
-    if (CaptureGameMode == NULL || CapturePlayerController == NULL)
+    if ((CaptureGameMode == nullptr) || (CapturePlayerController == nullptr))
     {
         UE_LOG(LogStereoPanorama, Warning, TEXT("Missing GameMode or PlayerController"));
         return;
@@ -131,5 +136,4 @@ void FStereoPanoramaManager::PanoramicTogglePause(const TArray<FString>& Args)
     {
         CaptureGameMode->SetPause(CapturePlayerController);
     }
-
-}
\ No newline at end of file
+}
diff --git a/Engine/Plugins/Experimental/StereoPanorama/Source/StereoPanorama/Private/StereoPanoramaManager.h b/Engine/Plugins/Experimental/StereoPanorama/Source/StereoPanorama/Private/StereoPanoramaManager.h
--- a/Engine/Plugins/Experimental/StereoPanorama/Source/StereoPanorama/Private/StereoPanoramaManager.h
+++ b/Engine/Plugins/Experimental/StereoPanorama/Source/StereoPanorama/Private/StereoPanoramaManager.h
@@ -1,22 +1,45 @@
 // Copyright 2015 Kite & Lightning.  All rights reserved.
 
 #pragma once
+
 #include "SceneCapturer.h"
 
+
 class FStereoPanoramaManager
 {
-private:
-	FAutoConsoleCommand PanoramicScreenshotCommand;
-	FAutoConsoleCommand PanoramicMovieCommand;
-	FAutoConsoleCommand PanoramicQualityCommand;
-    FAutoConsoleCommand PanoramicPauseCommand;
+public:
+
+	FStereoPanoramaManager()
+		: PanoramicScreenshotCommand(
+			TEXT("SP.PanoramicScreenshot"),
+			*NSLOCTEXT("StereoPanorama", "CommandText_ScreenShot", "Takes a panoramic screenshot").ToString(),
+			FConsoleCommandWithArgsDelegate::CreateRaw(this, &FStereoPanoramaManager::PanoramicScreenshot))
+		, PanoramicMovieCommand(
+			TEXT("SP.PanoramicMovie"),
+			*NSLOCTEXT("StereoPanorama", "CommandText_MovieCapture", "Takes a sequence of panoramic screenshots").ToString(),
+			FConsoleCommandWithArgsDelegate::CreateRaw(this, &FStereoPanoramaManager::PanoramicMovie))
+		, PanoramicQualityCommand(
+			TEXT("SP.PanoramicQuality"),
+			*NSLOCTEXT("StereoPanorama", "CommandText_Quality", "Sets the quality of the panoramic screenshot to 'preview | average | improved'").ToString(),
+			FConsoleCommandWithArgsDelegate::CreateRaw(this, &FStereoPanoramaManager::PanoramicQuality))
+		, PanoramicPauseCommand(
+            TEXT("SP.TogglePause"),
+			*NSLOCTEXT("StereoPanorama", "CommandText_PauseGame", "Toggles Pausing/Unpausing of the game through StereoPanorama Plugin").ToString(),
+			FConsoleCommandWithArgsDelegate::CreateRaw(this, &FStereoPanoramaManager::PanoramicTogglePause))
+		, SceneCapturer(nullptr)
+	{ }
+
+public:
+
+	void Cleanup();
+
+	void PanoramicScreenshot(const TArray<FString>& Args);
+	void PanoramicScreenshot(const int32 InStartFrame, const int32 InEndFrame, FStereoCaptureDoneDelegate& InStereoCaptureDoneDelegate);
+	void PanoramicMovie(const TArray<FString>& Args);
+	void PanoramicQuality(const TArray<FString>& Args);
+	void PanoramicTogglePause(const TArray<FString>& Args);
 
 public:
-    void PanoramicScreenshot( const TArray<FString>& Args );
-	void PanoramicScreenshot( const int32 InStartFrame, const int32 InEndFrame, FStereoCaptureDoneDelegate& InStereoCaptureDoneDelegate);
-	void PanoramicMovie( const TArray<FString>& Args );
-	void PanoramicQuality( const TArray<FString>& Args );
-    void PanoramicTogglePause(const TArray<FString>& Args);
 
 	static IConsoleVariable* HorizontalAngularIncrement;
 	static IConsoleVariable* VerticalAngularIncrement;
@@ -25,38 +48,25 @@ public:
 	static IConsoleVariable* ForceAlpha;
 	static IConsoleVariable* GenerateDebugImages;
 	static IConsoleVariable* ConcurrentCaptures;
-    static IConsoleVariable* CaptureHorizontalFOV;
-    static IConsoleVariable* CaptureSlicePixelWidth;
-    static IConsoleVariable* EnableBilerp;
-    static IConsoleVariable* SuperSamplingMethod;
-    static IConsoleVariable* OutputDir;
-    static IConsoleVariable* ShouldOverrideInitialYaw;
-    static IConsoleVariable* ForcedInitialYaw;
-    static IConsoleVariable* FadeStereoToZeroAtSides;
-
-	FStereoPanoramaManager()
-		: PanoramicScreenshotCommand(
-			TEXT( "SP.PanoramicScreenshot" ),
-			*NSLOCTEXT( "StereoPanorama", "CommandText_ScreenShot", "Takes a panoramic screenshot" ).ToString(),
-			FConsoleCommandWithArgsDelegate::CreateRaw( this, &FStereoPanoramaManager::PanoramicScreenshot ) )
-		, PanoramicMovieCommand(
-			TEXT( "SP.PanoramicMovie" ),
-			*NSLOCTEXT( "StereoPanorama", "CommandText_MovieCapture", "Takes a sequence of panoramic screenshots" ).ToString(),
-			FConsoleCommandWithArgsDelegate::CreateRaw( this, &FStereoPanoramaManager::PanoramicMovie ) )
-		, PanoramicQualityCommand(
-			TEXT( "SP.PanoramicQuality" ),
-			*NSLOCTEXT( "StereoPanorama", "CommandText_Quality", "Sets the quality of the panoramic screenshot to 'preview | average | improved'" ).ToString(),
-			FConsoleCommandWithArgsDelegate::CreateRaw( this, &FStereoPanoramaManager::PanoramicQuality ) )
-		, PanoramicPauseCommand(
-            TEXT( "SP.TogglePause" ),
-			*NSLOCTEXT( "StereoPanorama", "CommandText_PauseGame", "Toggles Pausing/Unpausing of the game through StereoPanorama Plugin" ).ToString(),
-			FConsoleCommandWithArgsDelegate::CreateRaw( this, &FStereoPanoramaManager::PanoramicTogglePause ) )
-		, SceneCapturer( NULL )
-	{
-	}
+	static IConsoleVariable* CaptureHorizontalFOV;
+	static IConsoleVariable* CaptureSlicePixelWidth;
+	static IConsoleVariable* EnableBilerp;
+	static IConsoleVariable* SuperSamplingMethod;
+	static IConsoleVariable* OutputDir;
+	static IConsoleVariable* ShouldOverrideInitialYaw;
+	static IConsoleVariable* ForcedInitialYaw;
+	static IConsoleVariable* FadeStereoToZeroAtSides;
 
-	void Cleanup();
+public:
 
+	/** The scene capturer object. */
 	UPROPERTY()
 	class USceneCapturer* SceneCapturer;
-};
\ No newline at end of file
+
+private:
+
+	FAutoConsoleCommand PanoramicScreenshotCommand;
+	FAutoConsoleCommand PanoramicMovieCommand;
+	FAutoConsoleCommand PanoramicQualityCommand;
+	FAutoConsoleCommand PanoramicPauseCommand;
+};
diff --git a/Engine/Plugins/Experimental/StereoPanorama/Source/StereoPanorama/Private/StereoStaticMeshComponent.cpp b/Engine/Plugins/Experimental/StereoPanorama/Source/StereoPanorama/Private/StereoStaticMeshComponent.cpp
--- a/Engine/Plugins/Experimental/StereoPanorama/Source/StereoPanorama/Private/StereoStaticMeshComponent.cpp
+++ b/Engine/Plugins/Experimental/StereoPanorama/Source/StereoPanorama/Private/StereoStaticMeshComponent.cpp
@@ -4,9 +4,11 @@
 #include "StereoStaticMeshComponent.h"
 #include "StaticMeshResources.h"
 
-class FStereoStaticMeshSceneProxy : public FStaticMeshSceneProxy
+
+class FStereoStaticMeshSceneProxy
+	: public FStaticMeshSceneProxy
 {
-    ESPStereoCameraLayer::Type EyeToRender;
+    ESPStereoCameraLayer EyeToRender;
 
 public:
 
@@ -24,14 +26,14 @@ public:
         switch (View->StereoPass)
         {
         case eSSP_RIGHT_EYE:
-            if (EyeToRender != ESPStereoCameraLayer::RightEye && EyeToRender != ESPStereoCameraLayer::BothEyes)
+            if ((EyeToRender != ESPStereoCameraLayer::RightEye) && (EyeToRender != ESPStereoCameraLayer::BothEyes))
             {
                 bVisible = false;
             }
             break;
 
         case eSSP_LEFT_EYE:
-            if (EyeToRender != ESPStereoCameraLayer::LeftEye && EyeToRender != ESPStereoCameraLayer::BothEyes)
+            if ((EyeToRender != ESPStereoCameraLayer::LeftEye) && (EyeToRender != ESPStereoCameraLayer::BothEyes))
             {
                 bVisible = false;
             }
@@ -48,22 +50,20 @@ public:
         return viewRelevance;
 
     }
-
-
 };
 
+
 FPrimitiveSceneProxy* UStereoStaticMeshComponent::CreateSceneProxy()
 {
-    if (StaticMesh == NULL
-        || StaticMesh->RenderData == NULL
-        || StaticMesh->RenderData->LODResources.Num() == 0
-        || StaticMesh->RenderData->LODResources[0].VertexBuffer.GetNumVertices() == 0)
+    if ((StaticMesh == nullptr) ||
+		(StaticMesh->RenderData == nullptr) ||
+		(StaticMesh->RenderData->LODResources.Num() == 0) ||
+		(StaticMesh->RenderData->LODResources[0].VertexBuffer.GetNumVertices() == 0))
     {
-        return NULL;
+        return nullptr;
     }
 
     FPrimitiveSceneProxy* Proxy = ::new FStereoStaticMeshSceneProxy(this);
-    return Proxy;
-}
-
 
+	return Proxy;
+}
diff --git a/Engine/Plugins/Experimental/StereoPanorama/Source/StereoPanorama/Private/StereoStaticMeshComponent.h b/Engine/Plugins/Experimental/StereoPanorama/Source/StereoPanorama/Private/StereoStaticMeshComponent.h
--- a/Engine/Plugins/Experimental/StereoPanorama/Source/StereoPanorama/Private/StereoStaticMeshComponent.h
+++ b/Engine/Plugins/Experimental/StereoPanorama/Source/StereoPanorama/Private/StereoStaticMeshComponent.h
@@ -7,30 +7,29 @@
 #include "Components/StaticMeshComponent.h"
 #include "StereoStaticMeshComponent.generated.h"
 
-UENUM(BlueprintType)       //"BlueprintType" is essential specifier
-namespace ESPStereoCameraLayer
+
+UENUM(BlueprintType)
+enum class ESPStereoCameraLayer
 {
-    //256 entries max
-    enum Type
-    {
-        LeftEye,
-        RightEye,
-        BothEyes,
-    };
-}
+    LeftEye,
+    RightEye,
+    BothEyes
+};
+
 
 /**
  * 
  */
 UCLASS(ClassGroup = (Rendering, Common), hidecategories = (Object, Activation, "Components|Activation"), ShowCategories = (Mobility), editinlinenew, meta = (BlueprintSpawnableComponent))
-class UStereoStaticMeshComponent : public UStaticMeshComponent
+class UStereoStaticMeshComponent
+	: public UStaticMeshComponent
 {
     GENERATED_BODY()
 
 public:
 
     UPROPERTY(EditAnywhere, BlueprintReadWrite, Category = Sprite)
-    TEnumAsByte<ESPStereoCameraLayer::Type> EyeToRender;
+    TEnumAsByte<ESPStereoCameraLayer> EyeToRender;
 
     virtual FPrimitiveSceneProxy* CreateSceneProxy() override;
 	
diff --git a/Engine/Plugins/Experimental/StereoPanorama/Source/StereoPanorama/Public/StereoCapturePawn.h b/Engine/Plugins/Experimental/StereoPanorama/Source/StereoPanorama/Public/StereoCapturePawn.h
--- a/Engine/Plugins/Experimental/StereoPanorama/Source/StereoPanorama/Public/StereoCapturePawn.h
+++ b/Engine/Plugins/Experimental/StereoPanorama/Source/StereoPanorama/Public/StereoCapturePawn.h
@@ -6,14 +6,19 @@
 #include "LatentActions.h"
 #include "StereoCapturePawn.generated.h"
 
-class FStereoCaptureDoneAction : public FPendingLatentAction
+
+class FStereoCaptureDoneAction
+	: public FPendingLatentAction
 {
 public:
+
     FName ExecutionFunction;
     int32 OutputLink;
     FWeakObjectPtr CallbackTarget;
     bool IsStereoCaptureDone;
 
+public:
+
     FStereoCaptureDoneAction(const FLatentActionInfo& LatentInfo)
         : ExecutionFunction(LatentInfo.ExecutionFunction)
         , OutputLink(LatentInfo.Linkage)
@@ -36,11 +41,13 @@ public:
 #endif
 };
 
+
 /**
  * 
  */
 UCLASS(config = Game, Blueprintable, BlueprintType)
-class AStereoCapturePawn : public ADefaultPawn
+class AStereoCapturePawn
+	: public ADefaultPawn
 {
     GENERATED_BODY()
 
@@ -59,6 +66,4 @@ public:
     FStereoCaptureDoneAction* StereoCaptureDoneAction;
 
     void CopyAtlasDataToTextures(const TArray<FColor>& LeftEyeAtlasData, const TArray<FColor>& RightEyeAtlasData);
-
-	
 };
diff --git a/Engine/Plugins/Experimental/StereoPanorama/Source/StereoPanorama/Public/StereoPanorama.h b/Engine/Plugins/Experimental/StereoPanorama/Source/StereoPanorama/Public/StereoPanorama.h
--- a/Engine/Plugins/Experimental/StereoPanorama/Source/StereoPanorama/Public/StereoPanorama.h
+++ b/Engine/Plugins/Experimental/StereoPanorama/Source/StereoPanorama/Public/StereoPanorama.h
@@ -2,15 +2,21 @@
 
 #pragma once
 
-class FStereoPanoramaModule : public IModuleInterface
+
+/**
+ * Implements the StereoPanorama module.
+ */
+class FStereoPanoramaModule
+	: public IModuleInterface
 {
-private:
-	/** IModuleInterface - initialize the module */
-	virtual void StartupModule() override;
+public:
 
-	/** IModuleInterface - shutdown the module */
+	//~ IModuleInterface interface
+
+	virtual void StartupModule() override;
 	virtual void ShutdownModule() override;
 
 public:
+
 	static TSharedPtr<FStereoPanoramaManager> Get();
-};
\ No newline at end of file
+};
diff --git a/Engine/Plugins/Experimental/StereoPanorama/StereoPanorama.uplugin b/Engine/Plugins/Experimental/StereoPanorama/StereoPanorama.uplugin
--- a/Engine/Plugins/Experimental/StereoPanorama/StereoPanorama.uplugin
+++ b/Engine/Plugins/Experimental/StereoPanorama/StereoPanorama.uplugin
@@ -4,7 +4,7 @@
 	"VersionName" : "Alpha",
 	"FriendlyName" : "Stereo Panoramic Movie Capture",
 	"Description" : "A plugin to capture a sequence of panoramic images",
-	"Category" : "Recording",
+	"Category" : "Movie Capture",
 	"CreatedBy" : "Kite & Lightning",
 	"CreatedByURL" : "http://kiteandlightning.la/",
 	"DocsURL" : "",
diff --git a/Engine/Plugins/Media/MediaPlayerEditor/Source/MediaPlayerEditor/Private/Widgets/SMediaPlayerEditorViewer.cpp b/Engine/Plugins/Media/MediaPlayerEditor/Source/MediaPlayerEditor/Private/Widgets/SMediaPlayerEditorViewer.cpp
--- a/Engine/Plugins/Media/MediaPlayerEditor/Source/MediaPlayerEditor/Private/Widgets/SMediaPlayerEditorViewer.cpp
+++ b/Engine/Plugins/Media/MediaPlayerEditor/Source/MediaPlayerEditor/Private/Widgets/SMediaPlayerEditorViewer.cpp
@@ -614,7 +614,7 @@ FText SMediaPlayerEditorViewer::HandleRemainingTimeTextBlockText() const
 
 	if (Remaining <= FTimespan::Zero())
 	{
-		return FText::GetEmpty();
+		return FText::FromString(TEXT("--:--:--"));
 	}
 
 	return FText::AsTimespan(Remaining);
diff --git a/Engine/Plugins/Media/WmfMedia/Source/WmfMedia/Private/Player/WmfMediaPlayer.cpp b/Engine/Plugins/Media/WmfMedia/Source/WmfMedia/Private/Player/WmfMediaPlayer.cpp
--- a/Engine/Plugins/Media/WmfMedia/Source/WmfMedia/Private/Player/WmfMediaPlayer.cpp
+++ b/Engine/Plugins/Media/WmfMedia/Source/WmfMedia/Private/Player/WmfMediaPlayer.cpp
@@ -13,30 +13,12 @@ FWmfMediaPlayer::FWmfMediaPlayer()
 	, MediaSession(nullptr)
 {
 	Resolver = new(std::nothrow) FWmfMediaResolver;
-	{
-		Resolver->OnResolveComplete().BindLambda([=](TComPtr<IUnknown> SourceObject, FString ResolvedUrl) {
-			AsyncTask(ENamedThreads::GameThread, [=]() {
-				MediaEvent.Broadcast(
-					InitializeMediaSession(SourceObject, ResolvedUrl)
-						? EMediaEvent::MediaOpened
-						: EMediaEvent::MediaOpenFailed
-				);
-			});
-		});
-
-		Resolver->OnResolveFailed().BindLambda([=](FString FailedUrl) {
-			AsyncTask(ENamedThreads::GameThread, [=]() {
-				MediaEvent.Broadcast(EMediaEvent::MediaOpenFailed);
-			});
-		});
-	}
 }
 
 
 FWmfMediaPlayer::~FWmfMediaPlayer()
 {
 	Close();
-	Resolver->OnResolveComplete().Unbind();
 }
 
 
@@ -212,7 +194,7 @@ bool FWmfMediaPlayer::Open(const FString& Url)
 		return false;
 	}
 
-	return Resolver->ResolveUrl(Url);
+	return Resolver->ResolveUrl(Url, *this);
 }
 
 
@@ -223,7 +205,7 @@ bool FWmfMediaPlayer::Open(const TSharedRef<FArchive, ESPMode::ThreadSafe>& Arch
 		return false;
 	}
 
-	return Resolver->ResolveByteStream(Archive, OriginalUrl);
+	return Resolver->ResolveByteStream(Archive, OriginalUrl, *this);
 }
 
 
@@ -257,6 +239,31 @@ bool FWmfMediaPlayer::SetRate(float Rate)
 }
 
 
+/* IWmfMediaResolverCallbacks interface
+ *****************************************************************************/
+
+void FWmfMediaPlayer::ProcessResolveComplete(TComPtr<IUnknown> SourceObject, FString ResolvedUrl)
+{
+	// forward event to game thread
+	AsyncTask(ENamedThreads::GameThread, [=]() {
+		MediaEvent.Broadcast(
+			InitializeMediaSession(SourceObject, ResolvedUrl)
+				? EMediaEvent::MediaOpened
+				: EMediaEvent::MediaOpenFailed
+		);
+	});
+}
+
+
+void FWmfMediaPlayer::ProcessResolveFailed(FString FailedUrl)
+{
+	// forward event to game thread
+	AsyncTask(ENamedThreads::GameThread, [=]() {
+		MediaEvent.Broadcast(EMediaEvent::MediaOpenFailed);
+	});
+}
+
+
 /* FWmfMediaPlayer implementation
  *****************************************************************************/
 
diff --git a/Engine/Plugins/Media/WmfMedia/Source/WmfMedia/Private/Player/WmfMediaPlayer.h b/Engine/Plugins/Media/WmfMedia/Source/WmfMedia/Private/Player/WmfMediaPlayer.h
--- a/Engine/Plugins/Media/WmfMedia/Source/WmfMedia/Private/Player/WmfMediaPlayer.h
+++ b/Engine/Plugins/Media/WmfMedia/Source/WmfMedia/Private/Player/WmfMediaPlayer.h
@@ -2,6 +2,7 @@
 
 #pragma once
 
+#include "IWmfMediaResolverCallbacks.h"
 #include "AllowWindowsPlatformTypes.h"
 
 
@@ -11,6 +12,7 @@
 class FWmfMediaPlayer
 	: public IMediaInfo
 	, public IMediaPlayer
+	, protected IWmfMediaResolverCallbacks
 {
 public:
 
@@ -58,6 +60,13 @@ public:
 		return MediaEvent;
 	}
 
+protected:
+
+	// IWmfMediaResolverCallbacks interface
+
+	virtual void ProcessResolveComplete(TComPtr<IUnknown> SourceObject, FString ResolvedUrl) override;
+	virtual void ProcessResolveFailed(FString FailedUrl) override;
+
 protected:
 
 	/**
diff --git a/Engine/Plugins/Media/WmfMedia/Source/WmfMedia/Private/Wmf/IWmfMediaResolverCallbacks.h b/Engine/Plugins/Media/WmfMedia/Source/WmfMedia/Private/Wmf/IWmfMediaResolverCallbacks.h
new file mode 100644
--- /dev/null
+++ b/Engine/Plugins/Media/WmfMedia/Source/WmfMedia/Private/Wmf/IWmfMediaResolverCallbacks.h
@@ -0,0 +1,32 @@
+// Copyright 1998-2015 Epic Games, Inc. All Rights Reserved.
+
+#pragma once
+
+
+/**
+ * Interface for classes that can handle media resolver callbacks.
+ */
+class IWmfMediaResolverCallbacks
+{
+public:
+
+	/**
+	 * Called when URL resolving has completed.
+	 *
+	 * @param SourceObject The media source object that was resolved.
+	 * @param ResolvedUrl The resolved media URL.
+	 */
+	virtual void ProcessResolveComplete(TComPtr<IUnknown> SourceObject, FString ResolvedUrl) = 0;
+
+	/**
+	 * Called when URL resolving has failed.
+	 *
+	 * @param FailedUrl The media URL that couldn't be resolved.
+	 */
+	virtual void ProcessResolveFailed(FString FailedUrl) = 0;
+
+public:
+
+	/** Virtual destructor. */
+	virtual ~IWmfMediaResolverCallbacks() { }
+};
diff --git a/Engine/Plugins/Media/WmfMedia/Source/WmfMedia/Private/Wmf/WmfMediaResolveState.h b/Engine/Plugins/Media/WmfMedia/Source/WmfMedia/Private/Wmf/WmfMediaResolveState.h
--- a/Engine/Plugins/Media/WmfMedia/Source/WmfMedia/Private/Wmf/WmfMediaResolveState.h
+++ b/Engine/Plugins/Media/WmfMedia/Source/WmfMedia/Private/Wmf/WmfMediaResolveState.h
@@ -2,6 +2,7 @@
 
 #pragma once
 
+#include "IWmfMediaResolverCallbacks.h"
 #include "AllowWindowsPlatformTypes.h"
 
 
@@ -34,12 +35,53 @@ struct FWmfMediaResolveState
 public:
 
 	/** Creates and initializes a new instance. */
-	FWmfMediaResolveState(EWmfMediaResolveType InType, const FString& InUrl)
-		: Type(InType)
+	FWmfMediaResolveState(EWmfMediaResolveType InType, const FString& InUrl, IWmfMediaResolverCallbacks& InCallbacks)
+		: Callbacks(&InCallbacks)
+		, Type(InType)
 		, Url(InUrl)
 		, RefCount(0)
 	{ }
 
+public:
+
+	/** Invalidate the resolve state (resolve events will no longer be forwarded). */
+	void Invalidate()
+	{
+		FScopeLock Lock(&CriticalSection);
+		Callbacks = nullptr;
+	}
+
+	/**
+	 * Forward the event for a completed media resolve.
+	 *
+	 * @param SourceObject The media source object that was resolved.
+	 * @param ResolvedUrl The resolved media URL.
+	 */
+	void ResolveComplete(TComPtr<IUnknown> SourceObject, FString ResolvedUrl)
+	{
+		FScopeLock Lock(&CriticalSection);
+
+		if (Callbacks != nullptr)
+		{
+			Callbacks->ProcessResolveComplete(SourceObject, ResolvedUrl);
+		}
+	}
+
+	/**
+	 * Forward the event for a failed media resolve.
+	 *
+	 * @param FailedUrl The media URL that couldn't be resolved.
+	 */
+	void ResolveFailed(FString FailedUrl)
+	{
+		FScopeLock Lock(&CriticalSection);
+
+		if (Callbacks != nullptr)
+		{
+			Callbacks->ProcessResolveFailed(FailedUrl);
+		}
+	}
+
 public:
 
 	// IUnknown interface
@@ -81,6 +123,12 @@ public:
 
 private:
 
+	/** Object that receives event callbacks. */
+	IWmfMediaResolverCallbacks* Callbacks;
+
+	/** Critical section for gating access to Callbacks. */
+	FCriticalSection CriticalSection;
+
 	/** Holds a reference counter for this instance. */
 	int32 RefCount;
 };
diff --git a/Engine/Plugins/Media/WmfMedia/Source/WmfMedia/Private/Wmf/WmfMediaResolver.cpp b/Engine/Plugins/Media/WmfMedia/Source/WmfMedia/Private/Wmf/WmfMediaResolver.cpp
--- a/Engine/Plugins/Media/WmfMedia/Source/WmfMedia/Private/Wmf/WmfMediaResolver.cpp
+++ b/Engine/Plugins/Media/WmfMedia/Source/WmfMedia/Private/Wmf/WmfMediaResolver.cpp
@@ -27,6 +27,12 @@ void FWmfMediaResolver::Cancel()
 		return;
 	}
 
+	if (ResolveState != NULL)
+	{
+		ResolveState->Invalidate();
+		ResolveState.Reset();
+	}
+
 	TComPtr<IUnknown> CancelCookieCopy = CancelCookie;
 
 	if (CancelCookieCopy != NULL)
@@ -34,12 +40,10 @@ void FWmfMediaResolver::Cancel()
 		SourceResolver->CancelObjectCreation(CancelCookieCopy);
 		CancelCookie.Reset();
 	}
-
-	ResolveState.Reset();
 }
 
 
-bool FWmfMediaResolver::ResolveByteStream(const TSharedRef<FArchive, ESPMode::ThreadSafe>& Archive, const FString& OriginalUrl)
+bool FWmfMediaResolver::ResolveByteStream(const TSharedRef<FArchive, ESPMode::ThreadSafe>& Archive, const FString& OriginalUrl, IWmfMediaResolverCallbacks& Callbacks)
 {
 	if (SourceResolver == NULL)
 	{
@@ -49,13 +53,20 @@ bool FWmfMediaResolver::ResolveByteStream(const TSharedRef<FArchive, ESPMode::Th
 	Cancel();
 
 	TComPtr<FWmfMediaByteStream> ByteStream = new FWmfMediaByteStream(Archive);
-	ResolveState = new(std::nothrow) FWmfMediaResolveState(EWmfMediaResolveType::ByteStream, OriginalUrl);
+	TComPtr<FWmfMediaResolveState> NewResolveState = new(std::nothrow) FWmfMediaResolveState(EWmfMediaResolveType::ByteStream, OriginalUrl, Callbacks);
 
-	return SUCCEEDED(SourceResolver->BeginCreateObjectFromByteStream(ByteStream, *OriginalUrl, MF_RESOLUTION_MEDIASOURCE, NULL, &CancelCookie, this, ResolveState));
+	if (FAILED(SourceResolver->BeginCreateObjectFromByteStream(ByteStream, *OriginalUrl, MF_RESOLUTION_MEDIASOURCE, NULL, &CancelCookie, this, NewResolveState)))
+	{
+		return false;
+	}
+
+	ResolveState = NewResolveState;
+
+	return true;
 }
 
 
-bool FWmfMediaResolver::ResolveUrl(const FString& Url)
+bool FWmfMediaResolver::ResolveUrl(const FString& Url, IWmfMediaResolverCallbacks& Callbacks)
 {
 	if (SourceResolver == NULL)
 	{
@@ -64,9 +75,16 @@ bool FWmfMediaResolver::ResolveUrl(const FString& Url)
 
 	Cancel();
 
-	ResolveState = new(std::nothrow) FWmfMediaResolveState(EWmfMediaResolveType::Url, Url);
+	TComPtr<FWmfMediaResolveState> NewResolveState = new(std::nothrow) FWmfMediaResolveState(EWmfMediaResolveType::Url, Url, Callbacks);
+
+	if (FAILED(SourceResolver->BeginCreateObjectFromURL(*Url, MF_RESOLUTION_MEDIASOURCE, NULL, &CancelCookie, this, NewResolveState)))
+	{
+		return false;
+	}
+
+	ResolveState = NewResolveState;
 
-	return SUCCEEDED(SourceResolver->BeginCreateObjectFromURL(*Url, MF_RESOLUTION_MEDIASOURCE, NULL, &CancelCookie, this, ResolveState));
+	return true;
 }
 
 
@@ -121,11 +139,11 @@ STDMETHODIMP FWmfMediaResolver::Invoke(IMFAsyncResult* AsyncResult)
 
 	if (Succeeded)
 	{
-		ResolveCompletedDelegate.ExecuteIfBound(SourceObject, State->Url);
+		State->ResolveComplete(SourceObject, State->Url);
 	}
 	else
 	{
-		ResolveFailedDelegate.ExecuteIfBound(State->Url);
+		State->ResolveFailed(State->Url);
 	}
 
 	ResolveState.Reset();
diff --git a/Engine/Plugins/Media/WmfMedia/Source/WmfMedia/Private/Wmf/WmfMediaResolver.h b/Engine/Plugins/Media/WmfMedia/Source/WmfMedia/Private/Wmf/WmfMediaResolver.h
--- a/Engine/Plugins/Media/WmfMedia/Source/WmfMedia/Private/Wmf/WmfMediaResolver.h
+++ b/Engine/Plugins/Media/WmfMedia/Source/WmfMedia/Private/Wmf/WmfMediaResolver.h
@@ -5,11 +5,8 @@
 #include "AllowWindowsPlatformTypes.h"
 
 
-/** Delegate type that is executed when URL resolving has completed. */
-DECLARE_DELEGATE_TwoParams(FWmfResolveCompletedDelegate, TComPtr<IUnknown> /*SourceObject*/, FString /*ResolvedUrl*/);
-
-/** Delegate type that is executed when URL resolving has failed. */
-DECLARE_DELEGATE_OneParam(FWmfResolveFailedDelegate, FString /*FailedUrl*/);
+class IWmfMediaResolverCallbacks;
+struct FWmfMediaResolveState;
 
 
 /**
@@ -35,44 +32,26 @@ public:
 	 */
 	void Cancel();
 
-	/**
-	 * Get a delegate that is executed when URL resolving has completed.
-	 *
-	 * @return The delegate.
-	 */
-	FWmfResolveCompletedDelegate& OnResolveComplete()
-	{
-		return ResolveCompletedDelegate;
-	}
-
-	/**
-	 * Get a delegate that is executed when URL resolving has failed.
-	 *
-	 * @return The delegate.
-	 */
-	FWmfResolveFailedDelegate& OnResolveFailed()
-	{
-		return ResolveFailedDelegate;
-	}
-
 	/**
 	 * Resolves a media source from a byte stream.
 	 *
 	 * @param Archive The archive holding the media data.
 	 * @param OriginalUrl The original URL of the media that was loaded into the buffer.
+	 * @param Callbacks The object that handles resolver event callbacks.
 	 * @return true if the byte stream will be resolved, false otherwise.
 	 * @see Cancel, ResolveUrl
 	 */
-	bool ResolveByteStream(const TSharedRef<FArchive, ESPMode::ThreadSafe>& Archive, const FString& OriginalUrl);
+	bool ResolveByteStream(const TSharedRef<FArchive, ESPMode::ThreadSafe>& Archive, const FString& OriginalUrl, IWmfMediaResolverCallbacks& Callbacks);
 
 	/**
 	 * Resolves a media source from a file or internet URL.
 	 *
 	 * @param Url The URL of the media to open (file name or web address).
+	 * @param Callbacks The object that handles resolver event callbacks.
 	 * @return true if the URL will be resolved, false otherwise.
 	 * @see Cancel, ResolveByteStream
 	 */
-	bool ResolveUrl(const FString& Url);
+	bool ResolveUrl(const FString& Url, IWmfMediaResolverCallbacks& Callbacks);
 
 public:
 
@@ -97,14 +76,6 @@ private:
 
 	/** Holds the actual source resolver. */
 	TComPtr<IMFSourceResolver> SourceResolver;
-
-private:
-
-	/** A delegate that is executed when URL resolving has completed. */
-	FWmfResolveCompletedDelegate ResolveCompletedDelegate;
-
-	/** A delegate that is executed when URL resolving has failed. */
-	FWmfResolveFailedDelegate ResolveFailedDelegate;
 };
 
 
diff --git a/Engine/Plugins/Media/WmfMedia/Source/WmfMedia/Private/WmfMediaModule.cpp b/Engine/Plugins/Media/WmfMedia/Source/WmfMedia/Private/WmfMediaModule.cpp
--- a/Engine/Plugins/Media/WmfMedia/Source/WmfMedia/Private/WmfMediaModule.cpp
+++ b/Engine/Plugins/Media/WmfMedia/Source/WmfMedia/Private/WmfMediaModule.cpp
@@ -144,9 +144,9 @@ public:
 	{
 		const FString Extension = FPaths::GetExtension(Url);
 
-		if (!Extension.IsEmpty())
+		if (!Extension.IsEmpty() && SupportedFileTypes.Contains(Extension))
 		{
-			return SupportedFileTypes.Contains(Extension);
+			return true;
 		}
 
 		for (const FString& Scheme : SupportedUriSchemes)
diff --git a/Engine/Plugins/MovieScene/LevelSequenceEditor/LevelSequenceEditor.uplugin b/Engine/Plugins/MovieScene/LevelSequenceEditor/LevelSequenceEditor.uplugin
--- a/Engine/Plugins/MovieScene/LevelSequenceEditor/LevelSequenceEditor.uplugin
+++ b/Engine/Plugins/MovieScene/LevelSequenceEditor/LevelSequenceEditor.uplugin
@@ -19,7 +19,7 @@
 		{
 			"Name" : "LevelSequenceEditor",
 			"Type" : "Editor",
-			"LoadingPhase" : "Default"
+			"LoadingPhase" : "PostEngineInit"
 		}
 	]
 }
\ No newline at end of file
diff --git a/Engine/Plugins/MovieScene/LevelSequenceEditor/Source/LevelSequenceEditor/LevelSequenceEditor.Build.cs b/Engine/Plugins/MovieScene/LevelSequenceEditor/Source/LevelSequenceEditor/LevelSequenceEditor.Build.cs
--- a/Engine/Plugins/MovieScene/LevelSequenceEditor/Source/LevelSequenceEditor/LevelSequenceEditor.Build.cs
+++ b/Engine/Plugins/MovieScene/LevelSequenceEditor/Source/LevelSequenceEditor/LevelSequenceEditor.Build.cs
@@ -10,11 +10,13 @@ public class LevelSequenceEditor : ModuleRules
             new string[] {
 				"AssetTools",
 				"SceneOutliner",
+				"PlacementMode",
 			}
         );
         
         PrivateDependencyModuleNames.AddRange(
 			new string[] {
+				"AppFramework",
                 "LevelSequence",
 				"BlueprintGraph",
 				"Core",
@@ -24,6 +26,7 @@ public class LevelSequenceEditor : ModuleRules
                 "InputCore",
 				"LevelEditor",
 				"MovieScene",
+                "MovieSceneTools",
 				"MovieSceneTracks",
                 "PropertyEditor",
 				"Sequencer",
@@ -36,7 +39,9 @@ public class LevelSequenceEditor : ModuleRules
 		PrivateIncludePathModuleNames.AddRange(
 			new string[] {
 				"AssetTools",
+                "MovieSceneTools",
 				"SceneOutliner",
+				"PlacementMode",
 			}
 		);
 
diff --git a/Engine/Plugins/MovieScene/LevelSequenceEditor/Source/LevelSequenceEditor/Private/CinematicLevelViewportLayout.cpp b/Engine/Plugins/MovieScene/LevelSequenceEditor/Source/LevelSequenceEditor/Private/CinematicLevelViewportLayout.cpp
new file mode 100644
--- /dev/null
+++ b/Engine/Plugins/MovieScene/LevelSequenceEditor/Source/LevelSequenceEditor/Private/CinematicLevelViewportLayout.cpp
@@ -0,0 +1,170 @@
+// Copyright 1998-2015 Epic Games, Inc. All Rights Reserved.
+
+#include "LevelSequenceEditorPCH.h"
+#include "CinematicLevelViewportLayout.h"
+#include "LevelViewportLayout.h"
+#include "SLevelViewport.h"
+#include "CinematicViewport/SCinematicLevelViewport.h"
+#include "LevelViewportActions.h"
+
+class FCinematicViewportLayoutEntity : public IViewportLayoutEntity
+{
+public:
+	FCinematicViewportLayoutEntity(TSharedRef<SWidget> InWidget)
+		: Widget(InWidget)
+	{}
+
+	virtual TSharedRef<SWidget> AsWidget() const override { return Widget; }
+
+	bool IsPlayInEditorViewportActive() const { return false; }
+	void RegisterGameViewportIfPIE(){}
+	void SetKeyboardFocus(){ FSlateApplication::Get().SetKeyboardFocus(Widget); }
+	void OnLayoutDestroyed(){}
+	void SaveConfig(const FString&){}
+
+protected:
+
+	/** This entity's widget */
+	TSharedRef<SWidget> Widget;
+};
+
+TSharedRef<SWidget> FCinematicLevelViewportLayout_OnePane::MakeViewportLayout(const FString& LayoutString)
+{
+	FString SpecificLayoutString = GetTypeSpecificLayoutString(LayoutString);
+
+	TSharedPtr<SCinematicLevelViewport> ViewportWidget;
+
+	FName CinematicKey = "Cinematic";
+
+	WidgetSlot = SNew(SBox)
+	[
+		SAssignNew(ViewportWidget, SCinematicLevelViewport)
+		.ParentLayout(AsShared())
+		.LayoutName(CinematicKey)
+		.RevertToLayoutName(LevelViewportConfigurationNames::OnePane)
+		.bShowMaximize(false)
+	];
+
+	Viewports.Add( CinematicKey, MakeShareable( new FCinematicViewportLayoutEntity(ViewportWidget.ToSharedRef()) ) );
+
+	// Make newly-created perspective viewports active by default
+	GCurrentLevelEditingViewportClient = ViewportWidget->GetViewportClient().Get();
+
+	return WidgetSlot.ToSharedRef();
+}
+
+void FCinematicLevelViewportLayout_OnePane::SaveLayoutString(const FString& LayoutString) const
+{
+}
+
+const FName& FCinematicLevelViewportLayout_OnePane::GetLayoutTypeName() const
+{
+	static FName LayoutTypeName("OnePaneCinematic");
+	return LayoutTypeName;
+}
+
+void FCinematicLevelViewportLayout_OnePane::ReplaceWidget( TSharedRef<SWidget> Source, TSharedRef<SWidget> Replacement)
+{
+	checkf(Source == WidgetSlot->GetChildren()->GetChildAt(0), TEXT("Source widget should have already been a content widget for the splitter"));
+	WidgetSlot->SetContent(Replacement);
+}
+
+TSharedRef<SWidget> FCinematicLevelViewportLayout_TwoPane::MakeViewportLayout(const FString& LayoutString)
+{
+	FString SpecificLayoutString = GetTypeSpecificLayoutString(LayoutString);
+
+	TSharedPtr<SLevelViewport> ViewportWidget0;
+	TSharedPtr<SCinematicLevelViewport> ViewportWidget1;
+
+	FString ViewportKey0;
+	float SplitterPercentage = 0.5f;
+
+	if (!SpecificLayoutString.IsEmpty())
+	{
+		// The Layout String only holds the unique ID of the Additional Layout Configs to use
+		const FString& IniSection = FLayoutSaveRestore::GetAdditionalLayoutConfigIni();
+
+		ViewportKey0 = SpecificLayoutString + TEXT(".Viewport0");
+
+		FString PercentageString;
+		if (GConfig->GetString(*IniSection, *(SpecificLayoutString + TEXT(".Percentage")), PercentageString, GEditorPerProjectIni))
+		{
+			TTypeFromString<float>::FromString(SplitterPercentage, *PercentageString);
+		}
+	}
+
+	FName CinematicKey = "Cinematic";
+
+	SplitterWidget = 
+		SNew( SSplitter )
+		.Orientation(EOrientation::Orient_Horizontal)
+
+		+SSplitter::Slot()
+		.Value(SplitterPercentage)
+		[
+			SAssignNew(ViewportWidget0, SLevelViewport)
+			.ViewportType(LVT_Perspective)
+			.ParentLayout(AsShared())
+			.ParentLevelEditor(ParentLevelEditor)
+			.ConfigKey(ViewportKey0)
+		]
+
+		+SSplitter::Slot()
+		.Value(1.0f - SplitterPercentage)
+		[
+			SAssignNew(ViewportWidget1, SCinematicLevelViewport)
+			.ParentLayout(AsShared())
+			.LayoutName(CinematicKey)
+			.RevertToLayoutName(LevelViewportConfigurationNames::TwoPanesHoriz)
+		];
+
+	ViewportWidget0->GetLevelViewportClient().SetAllowCinematicPreview(false);
+
+	Viewports.Add( *ViewportKey0, MakeShareable( new FLevelViewportLayoutEntity(ViewportWidget0.ToSharedRef()) ) );
+	Viewports.Add( CinematicKey, MakeShareable( new FCinematicViewportLayoutEntity(ViewportWidget1.ToSharedRef()) ) );
+
+	// Make newly-created perspective viewports active by default
+	GCurrentLevelEditingViewportClient = &ViewportWidget0->GetLevelViewportClient();
+
+	InitCommonLayoutFromString(SpecificLayoutString);
+
+	return SplitterWidget.ToSharedRef();
+}
+
+void FCinematicLevelViewportLayout_TwoPane::SaveLayoutString(const FString& LayoutString) const
+{
+	if (!bIsTransitioning)
+	{
+		FString SpecificLayoutString = GetTypeSpecificLayoutString(LayoutString);
+
+		const FString& IniSection = FLayoutSaveRestore::GetAdditionalLayoutConfigIni();
+
+		float Percentage = SplitterWidget->SlotAt(0).SizeValue.Get();
+		GConfig->SetString(*IniSection, *(SpecificLayoutString + TEXT(".Percentage")), *TTypeToString<float>::ToString(Percentage), GEditorPerProjectIni);
+
+		SaveCommonLayoutString(SpecificLayoutString);
+	}
+}
+
+const FName& FCinematicLevelViewportLayout_TwoPane::GetLayoutTypeName() const
+{
+	static FName LayoutTypeName("TwoPaneCinematic");
+	return LayoutTypeName;
+}
+
+void FCinematicLevelViewportLayout_TwoPane::ReplaceWidget( TSharedRef<SWidget> Source, TSharedRef<SWidget> Replacement)
+{
+	for (int32 SlotIdx = 0; SlotIdx < SplitterWidget->GetChildren()->Num(); SlotIdx++)
+	{
+		if  (SplitterWidget->GetChildren()->GetChildAt(SlotIdx) == Source)
+		{
+			SplitterWidget->SlotAt(SlotIdx)
+			[
+				Replacement
+			];
+			return;
+		}
+	}
+
+	checkf(false, TEXT("Source widget should have already been a content widget for the splitter"));
+}
diff --git a/Engine/Plugins/MovieScene/LevelSequenceEditor/Source/LevelSequenceEditor/Private/CinematicLevelViewportLayout.h b/Engine/Plugins/MovieScene/LevelSequenceEditor/Source/LevelSequenceEditor/Private/CinematicLevelViewportLayout.h
new file mode 100644
--- /dev/null
+++ b/Engine/Plugins/MovieScene/LevelSequenceEditor/Source/LevelSequenceEditor/Private/CinematicLevelViewportLayout.h
@@ -0,0 +1,45 @@
+// Copyright 1998-2015 Epic Games, Inc. All Rights Reserved.
+
+#pragma once
+
+#include "LevelViewportLayout.h"
+
+
+/** A viewport layout that comprises a normal editor viewport, and a cinematic preview viewport */
+class FCinematicLevelViewportLayout_OnePane : public FLevelViewportLayout
+{
+public:
+
+	virtual void SaveLayoutString(const FString& LayoutString) const override;
+	virtual const FName& GetLayoutTypeName() const override;
+
+protected:
+
+	virtual TSharedRef<SWidget> MakeViewportLayout(const FString& LayoutString) override;
+	virtual void ReplaceWidget( TSharedRef< SWidget > Source, TSharedRef< SWidget > Replacement ) override;
+
+
+private:
+	/** The widget */
+	TSharedPtr<SBox> WidgetSlot;
+};
+
+
+/** A viewport layout that comprises a normal editor viewport, and a cinematic preview viewport */
+class FCinematicLevelViewportLayout_TwoPane : public FLevelViewportLayout
+{
+public:
+
+	virtual void SaveLayoutString(const FString& LayoutString) const override;
+	virtual const FName& GetLayoutTypeName() const override;
+
+protected:
+
+	virtual TSharedRef<SWidget> MakeViewportLayout(const FString& LayoutString) override;
+	virtual void ReplaceWidget( TSharedRef< SWidget > Source, TSharedRef< SWidget > Replacement ) override;
+
+
+private:
+	/** The splitter widget */
+	TSharedPtr<class SSplitter> SplitterWidget;
+};
diff --git a/Engine/Plugins/MovieScene/LevelSequenceEditor/Source/LevelSequenceEditor/Private/CinematicViewport/FilmOverlays.cpp b/Engine/Plugins/MovieScene/LevelSequenceEditor/Source/LevelSequenceEditor/Private/CinematicViewport/FilmOverlays.cpp
new file mode 100644
--- /dev/null
+++ b/Engine/Plugins/MovieScene/LevelSequenceEditor/Source/LevelSequenceEditor/Private/CinematicViewport/FilmOverlays.cpp
@@ -0,0 +1,599 @@
+// Copyright 1998-2015 Epic Games, Inc. All Rights Reserved.
+
+#include "LevelSequenceEditorPCH.h"
+#include "FilmOverlays.h"
+#include "LevelSequenceEditorStyle.h"
+#include "Widgets/Colors/SColorPicker.h"
+#include "NumericTypeInterface.h"
+
+#define LOCTEXT_NAMESPACE "LevelSequenceEditorFilmOverlays"
+
+struct FFilmOverlay_None : IFilmOverlay
+{
+	FText GetDisplayName() const { return LOCTEXT("OverlayDisabled", "Disabled"); }
+	
+	const FSlateBrush* GetThumbnail() const { return FLevelSequenceEditorStyle::Get()->GetBrush("FilmOverlay.Disabled"); }
+
+	void Paint(const FGeometry& AllottedGeometry, const FSlateRect& MyClippingRect, FSlateWindowElementList& OutDrawElements, int32 LayerId) const
+	{
+	}
+};
+
+struct FFilmOverlay_Grid : IFilmOverlay
+{
+	FFilmOverlay_Grid(int32 InNumDivsH, int32 InNumDivsV)
+		: NumDivsH(InNumDivsH)
+		, NumDivsV(InNumDivsV)
+	{
+		BrushName = *FString::Printf(TEXT("FilmOverlay.%dx%dGrid"), NumDivsH, NumDivsV);
+	}
+
+	FText GetDisplayName() const
+	{
+		return FText::Format(
+			LOCTEXT("GridNameFormat", "Grid ({0}x{1})"),
+			FText::AsNumber(NumDivsH),
+			FText::AsNumber(NumDivsV)
+		);
+	}
+	
+	const FSlateBrush* GetThumbnail() const
+	{
+		return FLevelSequenceEditorStyle::Get()->GetBrush(BrushName);
+	}
+
+	void Paint(const FGeometry& AllottedGeometry, const FSlateRect& MyClippingRect, FSlateWindowElementList& OutDrawElements, int32 LayerId) const
+	{
+		const int32 IntervalH = AllottedGeometry.Size.X / NumDivsH;
+		const int32 IntervalV = AllottedGeometry.Size.Y / NumDivsV;
+
+		TArray<FVector2D> LinePoints;
+		LinePoints.SetNum(2);
+
+		for (int32 OffsetH = 1; OffsetH < NumDivsH; ++OffsetH)
+		{
+			LinePoints[0] = FVector2D(IntervalH*OffsetH, 0.f);
+			LinePoints[1] = FVector2D(IntervalH*OffsetH, AllottedGeometry.Size.Y);
+
+			FSlateDrawElement::MakeLines(
+				OutDrawElements,
+				LayerId,
+				AllottedGeometry.ToPaintGeometry(),
+				LinePoints,
+				MyClippingRect,
+				ESlateDrawEffect::None,
+				Tint,
+				false
+				);
+		}
+
+		for (int32 OffsetV = 1; OffsetV < NumDivsV; ++OffsetV)
+		{
+			LinePoints[0] = FVector2D(0.f, IntervalV*OffsetV);
+			LinePoints[1] = FVector2D(AllottedGeometry.Size.X, IntervalV*OffsetV);
+
+			FSlateDrawElement::MakeLines(
+				OutDrawElements,
+				LayerId,
+				AllottedGeometry.ToPaintGeometry(),
+				LinePoints,
+				MyClippingRect,
+				ESlateDrawEffect::None,
+				Tint,
+				false
+				);
+		}
+
+	}
+private:
+	FName BrushName;
+	int32 NumDivsH;
+	int32 NumDivsV;
+};
+
+struct FFilmOverlay_Rabatment : IFilmOverlay
+{
+	FText GetDisplayName() const { return LOCTEXT("RabatmentName", "Rabatment"); }
+	
+	const FSlateBrush* GetThumbnail() const { return FLevelSequenceEditorStyle::Get()->GetBrush("FilmOverlay.Rabatment"); }
+
+	void Paint(const FGeometry& AllottedGeometry, const FSlateRect& MyClippingRect, FSlateWindowElementList& OutDrawElements, int32 LayerId) const
+	{
+		FVector2D Position(AllottedGeometry.Size.X / 2, AllottedGeometry.Size.Y / 2);
+		float Size = FMath::Min(AllottedGeometry.Size.X, AllottedGeometry.Size.Y) * .1f;
+
+		TArray<FVector2D> LinePoints;
+		LinePoints.SetNum(2);
+
+		LinePoints[0] = FVector2D(AllottedGeometry.Size.Y, 0.f);
+		LinePoints[1] = FVector2D(AllottedGeometry.Size.Y, AllottedGeometry.Size.Y);
+		FSlateDrawElement::MakeLines(
+			OutDrawElements,
+			LayerId,
+			AllottedGeometry.ToPaintGeometry(),
+			LinePoints,
+			MyClippingRect,
+			ESlateDrawEffect::None,
+			Tint,
+			false
+			);
+
+		LinePoints[0] = FVector2D(AllottedGeometry.Size.X - AllottedGeometry.Size.Y, 0.f);
+		LinePoints[1] = FVector2D(AllottedGeometry.Size.X - AllottedGeometry.Size.Y, AllottedGeometry.Size.Y);
+		FSlateDrawElement::MakeLines(
+			OutDrawElements,
+			LayerId,
+			AllottedGeometry.ToPaintGeometry(),
+			LinePoints,
+			MyClippingRect,
+			ESlateDrawEffect::None,
+			Tint,
+			false
+			);
+
+	}
+};
+
+struct FFilmOverlay_Crosshair : IFilmOverlay
+{
+	FText GetDisplayName() const { return LOCTEXT("CrosshairName", "Crosshair"); }
+	const FSlateBrush* GetThumbnail() const { return FLevelSequenceEditorStyle::Get()->GetBrush("FilmOverlay.Crosshair"); }
+	void Paint(const FGeometry& AllottedGeometry, const FSlateRect& MyClippingRect, FSlateWindowElementList& OutDrawElements, int32 LayerId) const
+	{
+		FVector2D Position(AllottedGeometry.Size.X / 2, AllottedGeometry.Size.Y / 2);
+		float Size = FMath::Min(AllottedGeometry.Size.X, AllottedGeometry.Size.Y) * .1f;
+
+		TArray<FVector2D> LinePoints;
+		LinePoints.SetNum(2);
+
+		LinePoints[0] = FVector2D(Position.X, Position.Y - Size);
+		LinePoints[1] = FVector2D(Position.X, Position.Y - Size * .25f);
+		FSlateDrawElement::MakeLines(
+			OutDrawElements,
+			LayerId,
+			AllottedGeometry.ToPaintGeometry(),
+			LinePoints,
+			MyClippingRect,
+			ESlateDrawEffect::None,
+			Tint
+			);
+		LinePoints[0] = FVector2D(Position.X, Position.Y + Size);
+		LinePoints[1] = FVector2D(Position.X, Position.Y + Size * .25f);
+		FSlateDrawElement::MakeLines(
+			OutDrawElements,
+			LayerId,
+			AllottedGeometry.ToPaintGeometry(),
+			LinePoints,
+			MyClippingRect,
+			ESlateDrawEffect::None,
+			Tint,
+			false
+			);
+
+		LinePoints[0] = FVector2D(Position.X - Size, Position.Y);
+		LinePoints[1] = FVector2D(Position.X - Size * .25f, Position.Y);
+		FSlateDrawElement::MakeLines(
+			OutDrawElements,
+			LayerId,
+			AllottedGeometry.ToPaintGeometry(),
+			LinePoints,
+			MyClippingRect,
+			ESlateDrawEffect::None,
+			Tint
+			);
+		LinePoints[0] = FVector2D(Position.X + Size, Position.Y);
+		LinePoints[1] = FVector2D(Position.X + Size * .25f, Position.Y);
+		FSlateDrawElement::MakeLines(
+			OutDrawElements,
+			LayerId,
+			AllottedGeometry.ToPaintGeometry(),
+			LinePoints,
+			MyClippingRect,
+			ESlateDrawEffect::None,
+			Tint,
+			false
+			);
+	}
+};
+
+struct FFilmOverlay_SafeFrame : IFilmOverlay
+{
+	FFilmOverlay_SafeFrame(FText InDisplayName, float InSizePercentage, const FLinearColor& InColor)
+	{
+		DisplayName = InDisplayName;
+		SizePercentage = InSizePercentage;
+
+		Tint = InColor;
+	}
+
+	FText GetDisplayName() const { return DisplayName; }
+
+	const FSlateBrush* GetThumbnail() const { return nullptr; }
+
+	void Paint(const FGeometry& AllottedGeometry, const FSlateRect& MyClippingRect, FSlateWindowElementList& OutDrawElements, int32 LayerId) const
+	{
+		FVector2D TopLeft = AllottedGeometry.Size * ((100.f - SizePercentage) * .5f) / 100.f;
+		FVector2D BottomRight = AllottedGeometry.Size - TopLeft;
+
+		TArray<FVector2D> LinePoints;
+		LinePoints.Add(FVector2D(TopLeft.X,			TopLeft.Y));
+		LinePoints.Add(FVector2D(BottomRight.X,		TopLeft.Y));
+		LinePoints.Add(FVector2D(BottomRight.X,		BottomRight.Y));
+		LinePoints.Add(FVector2D(TopLeft.X,			BottomRight.Y));
+		LinePoints.Add(FVector2D(TopLeft.X-1,		TopLeft.Y-1));
+
+		FSlateDrawElement::MakeLines(
+			OutDrawElements,
+			LayerId,
+			AllottedGeometry.ToPaintGeometry(),
+			LinePoints,
+			MyClippingRect,
+			ESlateDrawEffect::None,
+			Tint,
+			false
+			);
+	}
+
+	virtual TSharedPtr<SWidget> ConstructSettingsWidget() override
+	{
+		struct FPercentageInterface : TDefaultNumericTypeInterface<float>
+		{
+			virtual FString ToString(const float& Value) const override
+			{
+				return TDefaultNumericTypeInterface::ToString(Value) + TEXT("%");
+			}
+			virtual TOptional<float> FromString(const FString& InString, const float& InExistingValue) override
+			{
+				return TDefaultNumericTypeInterface::FromString(InString.Replace(TEXT("%"), TEXT("")), InExistingValue);
+			}
+			virtual bool IsCharacterValid(TCHAR InChar) const override
+			{
+				return InChar == '%' || TDefaultNumericTypeInterface::IsCharacterValid(InChar);
+			}
+		};
+
+		auto GetValue = [this]{ return SizePercentage; };
+		auto SetValue = [this](float NewValue){ SizePercentage = NewValue; };
+		auto SetValueCommitted = [this](float NewValue, ETextCommit::Type){ SizePercentage = NewValue; };
+
+		return SNew(SHorizontalBox)
+
+			// + SHorizontalBox::Slot()
+			// .AutoWidth()
+			// [
+			// 	SNew(STextBlock)
+			// 	.Text(LOCTEXT("SizePercentage", "Screen Percentage:"))
+			// ]
+
+			+ SHorizontalBox::Slot()
+			.HAlign(HAlign_Right)
+			[
+				SNew(SSpinBox<float>)
+				.MinValue(1)
+				.MaxValue(99)
+				.Value_Lambda(GetValue)
+				.OnValueChanged_Lambda(SetValue)
+				.OnValueCommitted_Lambda(SetValueCommitted)
+				.TypeInterface(MakeShareable(new FPercentageInterface))
+			];
+	}
+
+private:
+	float SizePercentage;
+	FText DisplayName;
+};
+
+void SFilmOverlay::Construct(const FArguments& InArgs)
+{
+	FilmOverlays = InArgs._FilmOverlays;
+}
+
+int32 SFilmOverlay::OnPaint(const FPaintArgs& Args, const FGeometry& AllottedGeometry, const FSlateRect& MyClippingRect, FSlateWindowElementList& OutDrawElements, int32 LayerId, const FWidgetStyle& InWidgetStyle, bool bParentEnabled) const
+{
+	TArray<IFilmOverlay*> Overlays = FilmOverlays.Get();
+	for (IFilmOverlay* Overlay : Overlays)
+	{
+		Overlay->Paint(AllottedGeometry, MyClippingRect, OutDrawElements, LayerId);
+		++LayerId;
+	}
+
+	return LayerId;
+}
+
+void SFilmOverlayOptions::Construct(const FArguments& InArgs)
+{
+	MasterColorTint = FLinearColor(1.f, 1.f, 1.f, .5f);
+
+	MasterFilmOverlays.Add(NAME_None, TUniquePtr<IFilmOverlay>(new FFilmOverlay_None));
+	MasterFilmOverlays.Add("3x3Grid", TUniquePtr<IFilmOverlay>(new FFilmOverlay_Grid(3, 3)));
+	MasterFilmOverlays.Add("2x2Grid", TUniquePtr<IFilmOverlay>(new FFilmOverlay_Grid(2, 2)));
+	MasterFilmOverlays.Add("Crosshair", TUniquePtr<IFilmOverlay>(new FFilmOverlay_Crosshair));
+	MasterFilmOverlays.Add("Rabatment", TUniquePtr<IFilmOverlay>(new FFilmOverlay_Rabatment));
+
+	ToggleableOverlays.Add("ActionSafeFrame",	TUniquePtr<IFilmOverlay>(new FFilmOverlay_SafeFrame(LOCTEXT("ActionSafeFrame", "Action Safe"), 95.f, FLinearColor::Red)));
+	ToggleableOverlays.Add("TitleSafeFrame",	TUniquePtr<IFilmOverlay>(new FFilmOverlay_SafeFrame(LOCTEXT("TitleSafeFrame", "Title Safe"), 90.f, FLinearColor::Yellow)));
+	ToggleableOverlays.Add("CustomSafeFrame",	TUniquePtr<IFilmOverlay>(new FFilmOverlay_SafeFrame(LOCTEXT("CustomSafeFrame", "Custom Safe"), 85.f, FLinearColor::Green)));
+
+	OverlayWidget = SNew(SFilmOverlay)
+		.Visibility(EVisibility::HitTestInvisible)
+		.FilmOverlays(this, &SFilmOverlayOptions::GetActiveFilmOverlays);
+
+	ChildSlot
+	[
+		SNew(SComboButton)
+		.ButtonStyle(FEditorStyle::Get(), "HoverHintOnly")
+		.ForegroundColor(FEditorStyle::GetColor("InvertedForeground"))
+		.OnGetMenuContent(this, &SFilmOverlayOptions::GetMenuContent)
+		.ButtonContent()
+		[
+			SNew(SBox)
+			.WidthOverride(36)
+			.HeightOverride(24)
+			.ToolTipText(LOCTEXT("FilmOverlaysToolTip", "Displays a list of available film overlays to apply to this viewport."))
+			[
+				SNew(SImage)
+				.Image(this, &SFilmOverlayOptions::GetCurrentThumbnail)
+			]
+		]
+	];
+}
+
+FLinearColor SFilmOverlayOptions::GetMasterColorTint() const
+{
+	return MasterColorTint;
+}
+
+void SFilmOverlayOptions::SetMasterColorTint(FLinearColor Tint)
+{
+	MasterColorTint = Tint;
+
+	IFilmOverlay* Overlay = GetMasterFilmOverlay();
+	if (Overlay)
+	{
+		Overlay->SetTint(MasterColorTint);
+	}
+}
+
+TSharedRef<SWidget> SFilmOverlayOptions::GetMenuContent()
+{
+	return SNew(SGridPanel)
+
+		+ SGridPanel::Slot(0, 0)
+		[
+			SNew(SVerticalBox)
+
+			+ SVerticalBox::Slot()
+			.AutoHeight()
+			.Padding(0.f, 5.f)
+			[
+				SNew(SHeader)
+				[
+					SNew(STextBlock)
+					.Text(LOCTEXT("OverlaysHeader", "Composition Overlays"))
+				]
+			]
+
+			+ SVerticalBox::Slot()
+			.AutoHeight()
+			[
+				ConstructMasterOverlaysMenu()
+			]
+		]
+
+		+ SGridPanel::Slot(0, 1)
+		[
+			SNew(SVerticalBox)
+
+			+ SVerticalBox::Slot()
+			.AutoHeight()
+			[
+				SNew(SHeader)
+				[
+					SNew(STextBlock)
+					.Text(LOCTEXT("ColorPickerHeader", "Overlay Color Tint"))
+				]
+			]
+
+			+ SVerticalBox::Slot()
+			.AutoHeight()
+			.Padding(10)
+			.HAlign(HAlign_Center)
+			[
+				SNew(SColorPicker)
+				.TargetColorAttribute(this, &SFilmOverlayOptions::GetMasterColorTint)
+				.UseAlpha(true)
+				.DisplayInlineVersion(true)
+				.OnColorCommitted(this, &SFilmOverlayOptions::SetMasterColorTint)
+			]
+		]
+
+		// @todo: specific overlay settings to go here
+		// + SGridPanel::Slot(1, 0)
+		// [
+		// 
+		// ]
+
+		+ SGridPanel::Slot(1, 0)
+		//+ SGridPanel::Slot(1, 1)
+		[
+			SNew(SVerticalBox)
+
+			+ SVerticalBox::Slot()
+			.AutoHeight()
+			.Padding(0.f, 5.f)
+			[
+				SNew(SHeader)
+				[
+					SNew(STextBlock)
+					.Text(LOCTEXT("SafeFrameHeader", "Safe Frames"))
+				]
+			]
+			
+			+ SVerticalBox::Slot()
+			[
+				ConstructToggleableOverlaysMenu()
+			]
+		];
+}
+
+TSharedRef<SWidget> SFilmOverlayOptions::ConstructMasterOverlaysMenu()
+{
+	TSharedRef<SUniformGridPanel> OverlaysPanel = SNew(SUniformGridPanel).SlotPadding(10.f);
+
+	TArray<FName> OverlayNames;
+	MasterFilmOverlays.GenerateKeyArray(OverlayNames);
+
+	const int32 NumColumns = FMath::Log2(OverlayNames.Num() - 1);
+
+	int32 ColumnIndex = 0, RowIndex = 0;
+	for (int32 OverlayIndex = 0; OverlayIndex < OverlayNames.Num(); ++OverlayIndex)
+	{
+		IFilmOverlay& Overlay = *MasterFilmOverlays[OverlayNames[OverlayIndex]].Get();
+		OverlaysPanel->AddSlot(ColumnIndex, RowIndex)
+		[
+			SNew(SButton)
+			.ButtonStyle(FEditorStyle::Get(), "HoverHintOnly")
+			.OnClicked(this, &SFilmOverlayOptions::SetMasterFilmOverlay, OverlayNames[OverlayIndex])
+			[
+				SNew(SVerticalBox)
+
+				+ SVerticalBox::Slot()
+				.Padding(0.f, 4.f)
+				.AutoHeight()
+				.HAlign(HAlign_Center)
+				[
+					SNew(SBox)
+					.WidthOverride(36)
+					.HeightOverride(24)
+					[
+						SNew(SImage)
+						.Image(Overlay.GetThumbnail())
+					]
+				]
+
+				+ SVerticalBox::Slot()
+				.AutoHeight()
+				.HAlign(HAlign_Center)
+				.Padding(0.f, 0.f, 0.f, 4.f)
+				[
+					SNew(STextBlock)
+					.ColorAndOpacity(FEditorStyle::GetColor("DefaultForeground"))
+					.Text(Overlay.GetDisplayName())
+				]
+			]
+		];
+
+		if (++ColumnIndex >= NumColumns)
+		{
+			ColumnIndex = 0;
+			++RowIndex;
+		}
+	}
+
+	return OverlaysPanel;
+}
+
+TSharedRef<SWidget> SFilmOverlayOptions::ConstructToggleableOverlaysMenu()
+{
+	TSharedRef<SGridPanel> GridPanel = SNew(SGridPanel);
+
+	int32 Row = 0;
+	for (auto& Pair : ToggleableOverlays)
+	{
+		TSharedPtr<SWidget> Settings = Pair.Value->ConstructSettingsWidget();
+		if (!Settings.IsValid())
+		{
+			continue;
+		}
+
+		TUniquePtr<IFilmOverlay>* Ptr = &Pair.Value;
+
+		auto OnCheckStateChanged = [=](ECheckBoxState InNewState){ (*Ptr)->SetEnabled(InNewState == ECheckBoxState::Checked); };
+		auto IsChecked = [=]{ return (*Ptr)->IsEnabled() ? ECheckBoxState::Checked : ECheckBoxState::Unchecked; };
+		auto IsEnabled = [=]{ return (*Ptr)->IsEnabled(); };
+
+		Settings->SetEnabled(TAttribute<bool>::Create(TAttribute<bool>::FGetter::CreateLambda(IsEnabled)));
+
+		GridPanel->AddSlot(0, Row)
+		.VAlign(VAlign_Center)
+		.Padding(10.f, 5.f)
+		[
+			SNew(SCheckBox)
+			.OnCheckStateChanged_Lambda(OnCheckStateChanged)
+			.IsChecked_Lambda(IsChecked)
+			[
+				SNew(STextBlock)
+				.Text(Pair.Value->GetDisplayName())
+			]
+		];
+
+		GridPanel->AddSlot(1, Row)
+		.VAlign(VAlign_Center)
+		.HAlign(HAlign_Right)
+		.Padding(10.f, 5.f)
+		[
+			Settings.ToSharedRef()
+		];
+
+		++Row;
+	}
+
+	return GridPanel;
+}
+
+FReply SFilmOverlayOptions::SetMasterFilmOverlay(FName InName)
+{
+	CurrentMasterOverlay = InName;
+	IFilmOverlay* Overlay = GetMasterFilmOverlay();
+	if (Overlay)
+	{
+		Overlay->SetTint(MasterColorTint);
+	}
+	return FReply::Unhandled();
+}
+
+TSharedRef<SFilmOverlay> SFilmOverlayOptions::GetFilmOverlayWidget() const
+{
+	return OverlayWidget.ToSharedRef();
+}
+
+TArray<IFilmOverlay*> SFilmOverlayOptions::GetActiveFilmOverlays() const
+{
+	TArray<IFilmOverlay*> Overlays;
+
+	if (IFilmOverlay* Overlay = GetMasterFilmOverlay())
+	{
+		Overlays.Add(Overlay);
+	}
+
+	for (auto& Pair: ToggleableOverlays)
+	{
+		if (Pair.Value->IsEnabled())
+		{
+			Overlays.Add(Pair.Value.Get());
+		}
+	}
+
+	return Overlays;
+}
+
+const FSlateBrush* SFilmOverlayOptions::GetCurrentThumbnail() const
+{
+	if (!CurrentMasterOverlay.IsNone())
+	{
+		return MasterFilmOverlays[CurrentMasterOverlay]->GetThumbnail();
+	}
+
+	return FLevelSequenceEditorStyle::Get()->GetBrush("FilmOverlay.DefaultThumbnail");
+}
+
+IFilmOverlay* SFilmOverlayOptions::GetMasterFilmOverlay() const
+{
+	if (!CurrentMasterOverlay.IsNone())
+	{
+		return MasterFilmOverlays[CurrentMasterOverlay].Get();
+	}
+	return nullptr;
+}
+
+#undef LOCTEXT_NAMESPACE
\ No newline at end of file
diff --git a/Engine/Plugins/MovieScene/LevelSequenceEditor/Source/LevelSequenceEditor/Private/CinematicViewport/FilmOverlays.h b/Engine/Plugins/MovieScene/LevelSequenceEditor/Source/LevelSequenceEditor/Private/CinematicViewport/FilmOverlays.h
new file mode 100644
--- /dev/null
+++ b/Engine/Plugins/MovieScene/LevelSequenceEditor/Source/LevelSequenceEditor/Private/CinematicViewport/FilmOverlays.h
@@ -0,0 +1,123 @@
+// Copyright 1998-2015 Epic Games, Inc. All Rights Reserved.
+
+#pragma once
+
+/** Abstract base class that defines how an overlay should be drawn over the viewport */
+struct IFilmOverlay
+{
+	IFilmOverlay() : Tint(FLinearColor::White), bEnabled(false) {}
+
+	virtual ~IFilmOverlay() {};
+
+	/** Get a localized display name that is representative of this overlay */
+	virtual FText GetDisplayName() const = 0;
+
+	/** Get a slate thumbnail brush that is representative of this overlay. 36x24 recommended */
+	virtual const FSlateBrush* GetThumbnail() const = 0;
+
+	/** Construct a widget that controls this overlay's arbitrary settings. Only used for toggleable overlays. */
+	virtual TSharedPtr<SWidget> ConstructSettingsWidget() { return nullptr; }
+
+	/** Paint the overlay */
+	virtual void Paint(const FGeometry& AllottedGeometry, const FSlateRect& MyClippingRect, FSlateWindowElementList& OutDrawElements, int32 LayerId) const = 0;
+
+public:
+
+	/** Get/Set a custom tint for this overlay */
+	const FLinearColor& GetTint() const { return Tint; }
+	void SetTint(const FLinearColor& InTint) { Tint = InTint; }
+
+	/** Get/Set whether this overlay should be drawn */
+	bool IsEnabled() const { return bEnabled; }
+	void SetEnabled(bool bInEnabled) { bEnabled = bInEnabled; }
+
+protected:
+
+	/** Tint to apply to this overlay */
+	FLinearColor Tint;
+
+	/** Whether this overlay is enabled or not */
+	bool bEnabled;
+};
+
+/** A widget that sits on top of a viewport, and draws custom content */
+class SFilmOverlay : public SCompoundWidget
+{
+public:
+	SLATE_BEGIN_ARGS(SFilmOverlay){}
+
+		/** User provided array of overlays to draw */
+		SLATE_ATTRIBUTE(TArray<IFilmOverlay*>, FilmOverlays)
+
+	SLATE_END_ARGS()
+
+	/** Construct this widget */
+	void Construct(const FArguments& InArgs);
+
+	/** Paint this widget */
+	virtual int32 OnPaint(const FPaintArgs& Args, const FGeometry& AllottedGeometry, const FSlateRect& MyClippingRect, FSlateWindowElementList& OutDrawElements, int32 LayerId, const FWidgetStyle& InWidgetStyle, bool bParentEnabled) const override;
+
+private:
+	/** Attribute used once per frame to retrieve the film overlays to paint */
+	TAttribute<TArray<IFilmOverlay*>> FilmOverlays;
+};
+
+/** A custom widget that comprises a combo box displaying all available overlay options */
+class SFilmOverlayOptions : public SCompoundWidget
+{
+public:
+	SLATE_BEGIN_ARGS(SFilmOverlayOptions){}
+	SLATE_END_ARGS()
+
+	/** Construct this widget */
+	void Construct(const FArguments& InArgs);
+
+	/** Retrieve the actual overlay widget that this widget controls. Can be positioned in any other widget hierarchy. */
+	TSharedRef<SFilmOverlay> GetFilmOverlayWidget() const;
+
+private:
+
+	/** Generate menu content for the combo button */
+	TSharedRef<SWidget> GetMenuContent();
+
+	/** Construct the part of the menu that defines the set of film overlays */
+	TSharedRef<SWidget> ConstructMasterOverlaysMenu();
+
+	/** Construct the part of the menu that defines the set of toggleable overlays (currently just safe-frames) */
+	TSharedRef<SWidget> ConstructToggleableOverlaysMenu();
+
+private:
+
+	/** Get the thumbnail to be displayed on the combo box button */
+	const FSlateBrush* GetCurrentThumbnail() const;
+
+	/** Get the master film overlay ptr (may be nullptr) */
+	IFilmOverlay* GetMasterFilmOverlay() const;
+
+	/** Get an array of all enabled film overlays */
+	TArray<IFilmOverlay*> GetActiveFilmOverlays() const;
+
+	/** Set the current master overlay to the specified name */
+	FReply SetMasterFilmOverlay(FName InName);
+
+	/** Get/Set the color tint override for the current master overlay */
+	FLinearColor GetMasterColorTint() const;
+	void SetMasterColorTint(FLinearColor Tint);
+
+private:
+
+	/** Set of master film overlays (only one can be active at a time) */
+	TMap<FName, TUniquePtr<IFilmOverlay>> MasterFilmOverlays;
+
+	/** The name of the current master overlay */
+	FName CurrentMasterOverlay;
+
+	/** Color tint to apply to master overlays */
+	FLinearColor MasterColorTint;
+
+	/** A map of toggleable overlays (any number can be active at a time) */
+	TMap<FName, TUniquePtr<IFilmOverlay>> ToggleableOverlays;
+
+	/** The overlay widget we control */
+	TSharedPtr<SFilmOverlay> OverlayWidget;
+};
\ No newline at end of file
