Copying //UE4/Dev-Sequencer to Dev-Main (//UE4/Dev-Main)

diff --git a/Engine/Plugins/MovieScene/LevelSequenceEditor/Source/LevelSequenceEditor/Private/CinematicViewport/SCinematicLevelViewport.cpp b/Engine/Plugins/MovieScene/LevelSequenceEditor/Source/LevelSequenceEditor/Private/CinematicViewport/SCinematicLevelViewport.cpp
new file mode 100644
--- /dev/null
+++ b/Engine/Plugins/MovieScene/LevelSequenceEditor/Source/LevelSequenceEditor/Private/CinematicViewport/SCinematicLevelViewport.cpp
@@ -0,0 +1,1025 @@
+// Copyright 1998-2015 Epic Games, Inc. All Rights Reserved.
+
+#include "LevelSequenceEditorPCH.h"
+#include "SCinematicLevelViewport.h"
+#include "ISequencer.h"
+#include "SequencerCommands.h"
+#include "MovieScene.h"
+#include "MovieSceneCinematicShotTrack.h"
+#include "MovieSceneCinematicShotSection.h"
+#include "MovieSceneSubTrack.h"
+#include "MovieSceneSubSection.h"
+#include "SCinematicTransportRange.h"
+#include "FilmOverlays.h"
+#include "LevelViewportLayout.h"
+#include "LevelViewportTabContent.h"
+#include "LevelSequenceEditorStyle.h"
+#include "SWidgetSwitcher.h"
+#include "Camera/CineCameraComponent.h"
+#include "UnitConversion.h"
+#include "ScopedTransaction.h"
+
+#define LOCTEXT_NAMESPACE "SCinematicLevelViewport"
+
+template<typename T>
+struct SNonThrottledSpinBox : SSpinBox<T>
+{
+	virtual FReply OnMouseButtonDown(const FGeometry& MyGeometry, const FPointerEvent& MouseEvent) override
+	{
+		FReply Reply = SSpinBox<T>::OnMouseButtonDown(MyGeometry, MouseEvent);
+		if (Reply.IsEventHandled())
+		{
+			Reply.PreventThrottling();
+		}
+		return Reply;
+	}
+};
+
+struct FTypeInterfaceProxy : INumericTypeInterface<float>
+{
+	TSharedPtr<INumericTypeInterface<float>> Impl;
+
+	/** Convert the type to/from a string */
+	virtual FString ToString(const float& Value) const override
+	{
+		if (Impl.IsValid())
+		{
+			return Impl->ToString(Value);
+		}
+		return FString();
+	}
+
+	virtual TOptional<float> FromString(const FString& InString, const float& InExistingValue) override
+	{
+		if (Impl.IsValid())
+		{
+			return Impl->FromString(InString, InExistingValue);
+		}
+		return TOptional<float>();
+	}
+
+	/** Check whether the typed character is valid */
+	virtual bool IsCharacterValid(TCHAR InChar) const
+	{
+		if (Impl.IsValid())
+		{
+			return Impl->IsCharacterValid(InChar);
+		}
+		return false;
+	}
+};
+
+FCinematicViewportClient::FCinematicViewportClient()
+	: FLevelEditorViewportClient(nullptr)
+{
+	bDrawAxes = false;
+	bIsRealtime = true;
+	SetGameView(true);
+	SetAllowCinematicPreview(true);
+	bDisableInput = false;
+}
+
+class SPreArrangedBox : public SCompoundWidget
+{
+public:
+	DECLARE_DELEGATE_OneParam(FOnArrange, const FGeometry&);
+
+	SLATE_BEGIN_ARGS(SPreArrangedBox){}
+		SLATE_EVENT(FOnArrange, OnArrange)
+		SLATE_DEFAULT_SLOT(FArguments, Content)
+	SLATE_END_ARGS()
+
+	void Construct(const FArguments& InArgs)
+	{
+		OnArrange = InArgs._OnArrange;
+		ChildSlot
+		[
+			InArgs._Content.Widget
+		];
+	}
+
+	virtual void OnArrangeChildren(const FGeometry& AllottedGeometry, FArrangedChildren& ArrangedChildren) const override
+	{
+		OnArrange.ExecuteIfBound(AllottedGeometry);
+		SCompoundWidget::OnArrangeChildren(AllottedGeometry, ArrangedChildren);
+	}
+
+private:
+
+	FOnArrange OnArrange;
+};
+
+class SCinematicPreviewViewport : public SEditorViewport
+{
+public:
+
+	virtual TSharedRef<FEditorViewportClient> MakeEditorViewportClient()
+	{
+		return MakeShareable( new FCinematicViewportClient );
+	}
+};
+
+BEGIN_SLATE_FUNCTION_BUILD_OPTIMIZATION
+void SCinematicLevelViewport::Construct(const FArguments& InArgs)
+{
+	ParentLayout = InArgs._ParentLayout;
+	LayoutName = InArgs._LayoutName;
+	RevertToLayoutName = InArgs._RevertToLayoutName;
+
+	TypeInterfaceProxy = MakeShareable( new FTypeInterfaceProxy );
+
+	FLevelSequenceEditorToolkit::OnOpened().AddSP(this, &SCinematicLevelViewport::OnEditorOpened);
+
+	FLinearColor Gray(.3f, .3f, .3f, 1.f);
+
+	TSharedRef<SFilmOverlayOptions> FilmOverlayOptions = SNew(SFilmOverlayOptions);
+
+	TSharedRef<SWidget> MaximizeButton = SNullWidget::NullWidget;
+	if (InArgs._bShowMaximize)
+	{
+		MaximizeButton =
+			SNew(SButton)
+			.Cursor(EMouseCursor::Default)
+			.ButtonStyle(FEditorStyle::Get(), "HoverHintOnly")
+			.OnClicked(this, &SCinematicLevelViewport::OnToggleMaximize)
+			.ToolTipText(LOCTEXT("Maximize_ToolTip", "Maximizes or restores this viewport"))
+			[
+				SNew(SBox)
+				.VAlign(VAlign_Center)
+				.HAlign(HAlign_Center)
+				[
+					SNew(SImage)
+					.Image(this, &SCinematicLevelViewport::GetMaximizeImage)
+					.ColorAndOpacity(Gray)
+				]
+			];
+	}
+
+	ChildSlot
+	[
+		SNew(SBorder)
+		.BorderImage(FEditorStyle::GetBrush("BlackBrush"))
+		.ForegroundColor(Gray)
+		.Padding(5.f)
+		[
+			SNew(SVerticalBox)
+
+			+ SVerticalBox::Slot()
+			.AutoHeight()
+			[
+				SNew(SHorizontalBox)
+
+				+ SHorizontalBox::Slot()
+				.AutoWidth()
+				[
+					SNew(SComboButton)
+					.ButtonStyle(FEditorStyle::Get(), "HoverHintOnly")
+					.ForegroundColor(Gray)
+					.OnGetMenuContent(this, &SCinematicLevelViewport::MakeCameraMenu)
+					.ButtonContent()
+					[
+						SNew(SBox)
+						.WidthOverride(36)
+						.HeightOverride(24)
+						.ToolTipText(LOCTEXT("CameraSelectionTooltip", "Displays a list of available cameras to view through."))
+						[
+							SNew(SImage)
+							.Image(FLevelSequenceEditorStyle::Get()->GetBrush("LevelSequenceEditor.CinematicViewportCamera"))
+							.ColorAndOpacity(Gray)
+						]
+					]
+				]
+				
+				+ SHorizontalBox::Slot()
+				[
+					SNew(SSpacer)
+				]
+				
+				+ SHorizontalBox::Slot()
+				.AutoWidth()
+				[
+					FilmOverlayOptions
+				]
+
+				+ SHorizontalBox::Slot()
+				.AutoWidth()
+				[
+					MaximizeButton
+				]
+
+				+ SHorizontalBox::Slot()
+				.AutoWidth()
+				[
+					SNew(SButton)
+					.Cursor(EMouseCursor::Default)
+					.ButtonStyle(FEditorStyle::Get(), "HoverHintOnly")
+					.OnClicked(this, &SCinematicLevelViewport::OnRevertLayout)
+					.ToolTipText(LOCTEXT("Close_ToolTip", "Close this cinematic viewport, reverting it back to a standard viewport"))
+					[
+						SNew(SBox)
+						.VAlign(VAlign_Center)
+						.HAlign(HAlign_Center)
+						[
+							SNew(SImage)
+							.Image(FLevelSequenceEditorStyle::Get()->GetBrush("LevelSequenceEditor.CinematicViewportClose"))
+							.ColorAndOpacity(Gray)
+						]
+					]
+				]
+			]
+
+			+ SVerticalBox::Slot()
+			.Padding(20.f)
+			[
+				SNew(SPreArrangedBox)
+				.OnArrange(this, &SCinematicLevelViewport::CacheDesiredViewportSize)
+				[
+					SNew(SVerticalBox)
+
+					+ SVerticalBox::Slot()
+					[
+						SNew(SSpacer)
+					]
+
+					+ SVerticalBox::Slot()
+					.AutoHeight()
+					.HAlign(HAlign_Center)
+					[
+						SNew(SBox)
+						.HeightOverride(this, &SCinematicLevelViewport::GetDesiredViewportHeight)
+						.WidthOverride(this, &SCinematicLevelViewport::GetDesiredViewportWidth)
+						[
+							SNew(SOverlay)
+
+							+ SOverlay::Slot()
+							[
+								SAssignNew(ViewportWidget, SCinematicPreviewViewport)
+							]
+
+							+ SOverlay::Slot()
+							[
+								FilmOverlayOptions->GetFilmOverlayWidget()
+							]
+						]
+					]
+
+					+ SVerticalBox::Slot()
+					.AutoHeight()
+					.HAlign(HAlign_Center)
+					[
+						SAssignNew(ViewportControls, SBox)
+						.Visibility(this, &SCinematicLevelViewport::GetControlsVisibility)
+						.WidthOverride(this, &SCinematicLevelViewport::GetDesiredViewportWidth)
+						.Padding(FMargin(0.f, 10.f))
+						[
+							SNew(SHorizontalBox)
+
+							+ SHorizontalBox::Slot()
+							.HAlign(HAlign_Left)
+							[
+								SNew(STextBlock)
+								.ColorAndOpacity(Gray)
+								.Text_Lambda([=]{ return UIData.ShotName; })
+							]
+
+							+ SHorizontalBox::Slot()
+							.HAlign(HAlign_Center)
+							.AutoWidth()
+							[
+								SNew(STextBlock)
+								.ColorAndOpacity(Gray)
+								.Text_Lambda([=]{ return UIData.Filmback; })
+							]
+
+							+ SHorizontalBox::Slot()
+							.HAlign(HAlign_Right)
+							[
+								SNew(STextBlock)
+								.Font(FEditorStyle::GetFontStyle("Sequencer.FixedFont"))
+								.ColorAndOpacity(Gray)
+								.Text_Lambda([=]{ return UIData.Frame; })
+							]
+						]
+					]
+
+					+ SVerticalBox::Slot()
+					[
+						SNew(SSpacer)
+					]
+				]
+			]
+
+
+			+ SVerticalBox::Slot()
+			.AutoHeight()
+			[
+			
+				SNew(SWidgetSwitcher)
+				.WidgetIndex(this, &SCinematicLevelViewport::GetVisibleWidgetIndex)
+
+				+ SWidgetSwitcher::Slot()
+				[
+					SNew(SVerticalBox)
+
+					+ SVerticalBox::Slot()
+					.AutoHeight()
+					.Padding(5.f)
+					[
+						SAssignNew(TransportRange, SCinematicTransportRange)
+					]
+
+					+ SVerticalBox::Slot()
+					.AutoHeight()
+					.Padding(5.f, 0.f)
+					[
+						SNew(SHorizontalBox)
+
+						+ SHorizontalBox::Slot()
+						.VAlign(VAlign_Center)
+						.AutoWidth()
+						[
+							SNew(STextBlock)
+							.Font(FEditorStyle::GetFontStyle("Sequencer.FixedFont"))
+							.ColorAndOpacity(Gray)
+							.Text_Lambda([=]{ return UIData.MasterStartText; })
+						]
+
+						+ SHorizontalBox::Slot()
+						.VAlign(VAlign_Center)
+						.AutoWidth()
+						.Padding(FMargin(10.f, 0.f))
+						[
+							SNew(SNonThrottledSpinBox<float>)
+							.TypeInterface(TypeInterfaceProxy)
+							.Style(FEditorStyle::Get(), "Sequencer.HyperlinkSpinBox")
+							.Font(FEditorStyle::GetFontStyle("Sequencer.FixedFont"))
+							.OnValueCommitted(this, &SCinematicLevelViewport::OnTimeCommitted)
+							.OnValueChanged(this, &SCinematicLevelViewport::SetTime)
+							.OnEndSliderMovement(this, &SCinematicLevelViewport::SetTime)
+							.MinValue(this, &SCinematicLevelViewport::GetMinTime)
+							.MaxValue(this, &SCinematicLevelViewport::GetMaxTime)
+							.Value(this, &SCinematicLevelViewport::GetTime)
+						]
+
+						+ SHorizontalBox::Slot()
+						.VAlign(VAlign_Center)
+						.HAlign(HAlign_Right)
+						[
+							SNew(SButton)
+							.OnClicked(this, &SCinematicLevelViewport::SetPlaybackStart)
+							.ToolTipText(LOCTEXT("SetPlayStart_Tooltip", "Set playback start to the current position"))
+							.ButtonStyle(FEditorStyle::Get(), "NoBorder")
+							.ContentPadding(2.0f)
+							[
+								SNew(SImage)
+								.Image(FLevelSequenceEditorStyle::Get()->GetBrush("LevelSequenceEditor.CinematicViewportSetPlayStart"))
+							]
+						]
+
+						+ SHorizontalBox::Slot()
+						.VAlign(VAlign_Center)
+						.HAlign(HAlign_Center)
+						.AutoWidth()
+						[
+							SAssignNew(TransportControls, SBox)
+						]
+
+						+ SHorizontalBox::Slot()
+						.VAlign(VAlign_Center)
+						.HAlign(HAlign_Left)
+						[
+							SNew(SButton)
+							.OnClicked(this, &SCinematicLevelViewport::SetPlaybackEnd)
+							.ToolTipText(LOCTEXT("SetPlayEnd_Tooltip", "Set playback end to the current position"))
+							.ButtonStyle(FEditorStyle::Get(), "NoBorder")
+							.ContentPadding(2.0f)
+							[
+								SNew(SImage)
+								.Image(FLevelSequenceEditorStyle::Get()->GetBrush("LevelSequenceEditor.CinematicViewportSetPlayEnd"))
+							]
+						]
+
+						+ SHorizontalBox::Slot()
+						.VAlign(VAlign_Center)
+						.AutoWidth()
+						[
+							SNew(STextBlock)
+							.Font(FEditorStyle::GetFontStyle("Sequencer.FixedFont"))
+							.ColorAndOpacity(Gray)
+							.Text_Lambda([=]{ return UIData.MasterEndText; })
+						]
+					]
+				]
+
+				+ SWidgetSwitcher::Slot()
+				.VAlign(VAlign_Center)
+				.HAlign(HAlign_Center)
+				.Padding(5.f, 0.f)
+				[
+					SNew(STextBlock)
+					.ColorAndOpacity(Gray)
+					.Text(LOCTEXT("NoSequencerMessage", "No active Level Sequencer detected. Please edit a Level Sequence to enable full controls."))
+				]
+			]
+		]
+	];
+
+	FLevelSequenceEditorToolkit::IterateOpenToolkits([&](FLevelSequenceEditorToolkit& Toolkit){
+		Setup(Toolkit);
+		return false;
+	});
+}
+END_SLATE_FUNCTION_BUILD_OPTIMIZATION
+
+FReply SCinematicLevelViewport::OnRevertLayout()
+{
+	TSharedPtr<FLevelViewportLayout> LayoutPinned = ParentLayout.Pin();
+	if (LayoutPinned.IsValid())
+	{
+		TSharedPtr<FLevelViewportTabContent> ViewportTabPinned = LayoutPinned->GetParentTabContent().Pin();
+		if (ViewportTabPinned.IsValid())
+		{
+			// Viewport clients are going away.  Any current one is invalid.
+			GCurrentLevelEditingViewportClient = nullptr;
+			ViewportTabPinned->SetViewportConfiguration(RevertToLayoutName);
+			FSlateApplication::Get().DismissAllMenus();
+		}
+	}
+	return FReply::Handled();
+}
+
+int32 SCinematicLevelViewport::GetVisibleWidgetIndex() const
+{
+	return CurrentToolkit.IsValid() ? 0 : 1;
+}
+
+EVisibility SCinematicLevelViewport::GetControlsVisibility() const
+{
+	return CurrentToolkit.IsValid() ? EVisibility::Visible : EVisibility::Hidden;
+}
+
+FReply SCinematicLevelViewport::OnToggleMaximize()
+{
+	TSharedPtr<FLevelViewportLayout> ParentLayoutPinned = ParentLayout.Pin();
+	if (ParentLayoutPinned.IsValid() && ParentLayoutPinned->IsMaximizeSupported())
+	{
+		const bool bAllowAnimation = true;
+		const bool bMaximize = !ParentLayoutPinned->IsViewportMaximized(LayoutName);
+		
+		ParentLayoutPinned->RequestMaximizeViewport( LayoutName, bMaximize, false, bAllowAnimation );
+	}
+	return FReply::Handled();
+}
+
+bool SCinematicLevelViewport::IsMaximized() const
+{
+	TSharedPtr<FLevelViewportLayout> ParentLayoutPinned = ParentLayout.Pin();
+	if (ParentLayoutPinned.IsValid())
+	{
+		return ParentLayoutPinned->IsViewportMaximized(LayoutName);
+	}
+	return false;
+}
+
+const FSlateBrush* SCinematicLevelViewport::GetMaximizeImage() const
+{
+	return IsMaximized() ? FEditorStyle::GetBrush("LevelViewportToolBar.Maximize.Checked") : FEditorStyle::GetBrush("LevelViewportToolBar.Maximize.Normal");
+}
+
+TOptional<float> SCinematicLevelViewport::GetMinTime() const
+{
+	ISequencer* Sequencer = GetSequencer();
+	if (Sequencer)
+	{
+		return Sequencer->GetFocusedMovieSceneSequence()->GetMovieScene()->GetEditorData().WorkingRange.GetLowerBoundValue();
+	}
+	return TOptional<float>();
+}
+
+TOptional<float> SCinematicLevelViewport::GetMaxTime() const
+{
+	ISequencer* Sequencer = GetSequencer();
+	if (Sequencer)
+	{
+		return Sequencer->GetFocusedMovieSceneSequence()->GetMovieScene()->GetEditorData().WorkingRange.GetUpperBoundValue();
+	}
+	return TOptional<float>();
+}
+
+void SCinematicLevelViewport::OnTimeCommitted(float Value, ETextCommit::Type)
+{
+	SetTime(Value);
+}
+
+void SCinematicLevelViewport::SetTime(float Value)
+{
+	ISequencer* Sequencer = GetSequencer();
+	if (Sequencer)
+	{
+		Sequencer->SetGlobalTime(Value);
+	}
+}
+
+float SCinematicLevelViewport::GetTime() const
+{
+	ISequencer* Sequencer = GetSequencer();
+	if (Sequencer)
+	{
+		return Sequencer->GetGlobalTime();
+	}
+	return 0.f;
+}
+
+FReply SCinematicLevelViewport::SetPlaybackStart()
+{
+	ISequencer* Sequencer = GetSequencer();
+	if (Sequencer)
+	{
+		FScopedTransaction Transaction(LOCTEXT("SetPlaybackStart", "Set Playback Start"));
+
+		TRange<float> CurrentRange = Sequencer->GetFocusedMovieSceneSequence()->GetMovieScene()->GetPlaybackRange();
+
+		float NewPos = FMath::Min(Sequencer->GetGlobalTime(), CurrentRange.GetUpperBoundValue());
+		Sequencer->GetFocusedMovieSceneSequence()->GetMovieScene()->SetPlaybackRange(NewPos, CurrentRange.GetUpperBoundValue());
+	}
+	return FReply::Handled();
+}
+FReply SCinematicLevelViewport::SetPlaybackEnd()
+{
+	ISequencer* Sequencer = GetSequencer();
+	if (Sequencer)
+	{
+		FScopedTransaction Transaction(LOCTEXT("SetPlaybackEnd", "Set Playback End"));
+
+		TRange<float> CurrentRange = Sequencer->GetFocusedMovieSceneSequence()->GetMovieScene()->GetPlaybackRange();
+
+		float NewPos = FMath::Max(Sequencer->GetGlobalTime(), CurrentRange.GetLowerBoundValue());
+		Sequencer->GetFocusedMovieSceneSequence()->GetMovieScene()->SetPlaybackRange(CurrentRange.GetLowerBoundValue(), NewPos);
+	}
+	return FReply::Handled();
+}
+void SCinematicLevelViewport::CacheDesiredViewportSize(const FGeometry& AllottedGeometry)
+{
+	FVector2D ViewportControlsSize = ViewportControls->GetDesiredSize();
+
+	FVector2D AllowableSpace = AllottedGeometry.GetLocalSize();
+	AllowableSpace.Y -= ViewportControlsSize.Y;
+
+	TSharedPtr<FCinematicViewportClient> ViewportClient = GetViewportClient();
+	if (ViewportClient->IsAspectRatioConstrained())
+	{
+		const float MinSize = FMath::Min(AllowableSpace.X / ViewportClient->AspectRatio, AllowableSpace.Y);
+		DesiredViewportSize = FVector2D(int32(ViewportClient->AspectRatio * MinSize), int32(MinSize));
+	}
+	else
+	{
+		DesiredViewportSize = AllowableSpace;
+	}
+}
+
+FOptionalSize SCinematicLevelViewport::GetDesiredViewportWidth() const
+{
+	return DesiredViewportSize.X;
+}
+
+FOptionalSize SCinematicLevelViewport::GetDesiredViewportHeight() const
+{
+	return DesiredViewportSize.Y;
+}
+
+TSharedRef<SWidget> SCinematicLevelViewport::MakeCameraMenu()
+{
+	ISequencer* Sequencer = GetSequencer();
+	UMovieSceneSequence* Sequence = Sequencer ? Sequencer->GetFocusedMovieSceneSequence() : nullptr;
+
+	FMenuBuilder MenuBuilder(true, nullptr, nullptr);
+	{
+		// let toolkits populate the menu
+		MenuBuilder.AddMenuEntry(
+			LOCTEXT("FreeCamera", "Free Camera"),
+			LOCTEXT("FreeCamera_ToolTip", "Allows the arbitrary movement of the camera"),
+			FSlateIcon(),
+			FUIAction(
+				FExecuteAction::CreateSP(this, &SCinematicLevelViewport::SetCameraMode, ECinematicCameraMode::FreeCamera, TOptional<FGuid>()),
+				FCanExecuteAction::CreateLambda([=]{ return true; }),
+				FIsActionChecked::CreateLambda([=]{ return CameraMode == ECinematicCameraMode::FreeCamera; })
+			),
+			NAME_None,
+			EUserInterfaceActionType::ToggleButton
+		);
+
+		if (Sequence)
+		{
+			MenuBuilder.BeginSection("Sequence");
+			{
+				MenuBuilder.AddMenuEntry(
+					LOCTEXT("AutoCamera", "Auto"),
+					LOCTEXT("AutoCamera_ToolTip", "Looks through the camera defined by either the shot track, camera cut track, or any other available camera (in that order)"),
+					FSlateIcon(),
+					FUIAction(
+						FExecuteAction::CreateSP(this, &SCinematicLevelViewport::SetCameraMode, ECinematicCameraMode::Auto, TOptional<FGuid>()),
+						FCanExecuteAction::CreateLambda([=]{ return true; }),
+						FIsActionChecked::CreateLambda([=]{ return CameraMode == ECinematicCameraMode::Auto; })
+					),
+					NAME_None,
+					EUserInterfaceActionType::ToggleButton
+				);
+			}
+			MenuBuilder.EndSection();
+
+			MenuBuilder.BeginSection("Cameras");
+			{
+				for (auto& Pair : Cameras)
+				{
+					FGuid CameraGuid = Pair.Key;
+
+					MenuBuilder.AddMenuEntry(
+						FText::FromString(Pair.Value.DisplayName),
+						FText(),
+						FSlateIcon(),
+						FUIAction(
+							FExecuteAction::CreateSP(this, &SCinematicLevelViewport::SetCameraMode, ECinematicCameraMode::SpecificCamera, TOptional<FGuid>(CameraGuid)),
+							FCanExecuteAction::CreateLambda([=]{ return true; }),
+							FIsActionChecked::CreateLambda([=]{
+								return CameraMode == ECinematicCameraMode::SpecificCamera &&
+									LockedCameraGuid.IsSet() &&
+									LockedCameraGuid.GetValue() == CameraGuid;
+							})
+						),
+						NAME_None,
+						EUserInterfaceActionType::ToggleButton
+					);
+				}
+			}
+			MenuBuilder.EndSection();
+		}
+	}
+
+	return MenuBuilder.MakeWidget();
+}
+
+void SCinematicLevelViewport::SetCameraMode(ECinematicCameraMode InCameraMode, TOptional<FGuid> InSpecificCameraGuid)
+{
+	CameraMode = InCameraMode;
+
+	TSharedPtr<FCinematicViewportClient> ViewportClient = GetViewportClient();
+
+	// Reset everything
+	LockedCameraGuid = TOptional<FGuid>();
+	ViewportClient->SetActorLock(nullptr);
+
+	switch (CameraMode)
+	{
+	case ECinematicCameraMode::FreeCamera:
+		ViewportClient->bDisableInput = false;
+		break;
+
+	case ECinematicCameraMode::Auto:
+		ViewportClient->bDisableInput = true;
+		break;
+
+	case ECinematicCameraMode::SpecificCamera:
+		LockedCameraGuid = InSpecificCameraGuid;
+		ViewportClient->bDisableInput = true;
+		break;
+	}
+
+	ISequencer* Sequencer = GetSequencer();
+	if (Sequencer)
+	{
+		Sequencer->SetGlobalTime(Sequencer->GetGlobalTime());
+	}
+}
+
+void SCinematicLevelViewport::UpdateActiveCameras()
+{
+	Cameras.Reset();
+
+	ISequencer* Sequencer = GetSequencer();
+	if (!Sequencer)
+	{
+		return;
+	}
+
+	TSharedPtr<FMovieSceneSequenceInstance> SequenceInstance = Sequencer->GetFocusedMovieSceneSequenceInstance();
+	UMovieScene* MovieScene = SequenceInstance->GetSequence()->GetMovieScene();
+
+	const int32 NumSpawnables = MovieScene->GetSpawnableCount();
+	for (int32 SpawnableIndex = 0; SpawnableIndex < NumSpawnables; ++SpawnableIndex)
+	{
+		FMovieSceneSpawnable& Spawnable = MovieScene->GetSpawnable(SpawnableIndex);
+		if (Spawnable.GetClass()->IsChildOf(ACameraActor::StaticClass()))
+		{
+			ACameraActor* Camera = Cast<ACameraActor>(SequenceInstance->FindObject(Spawnable.GetGuid(), *Sequencer));
+			Cameras.Add(Spawnable.GetGuid(), FActiveCamera(Spawnable.GetName(), Camera));
+		}
+	}
+
+	const int32 NumPossessables = MovieScene->GetPossessableCount();
+	for (int32 PossessableIndex = 0; PossessableIndex < NumPossessables; ++PossessableIndex)
+	{
+		FMovieScenePossessable& Possessable = MovieScene->GetPossessable(PossessableIndex);
+		ACameraActor* Camera = Cast<ACameraActor>(SequenceInstance->FindObject(Possessable.GetGuid(), *Sequencer));
+		if (Camera)
+		{
+			Cameras.Add(Possessable.GetGuid(), FActiveCamera(Possessable.GetName(), Camera));
+		}
+	}
+}
+
+TSharedPtr<FCinematicViewportClient> SCinematicLevelViewport::GetViewportClient() const
+{
+	return StaticCastSharedPtr<FCinematicViewportClient>(ViewportWidget->GetViewportClient());
+}
+
+void SCinematicLevelViewport::OnUpdateCameraCut(UObject* CameraObject)
+{
+	if (CameraMode != ECinematicCameraMode::Auto)
+	{
+		return;
+	}
+
+	ACameraActor* CameraActor = Cast<ACameraActor>(CameraObject);
+	TSharedPtr<FCinematicViewportClient> ViewportClient = GetViewportClient();
+	
+	if (CameraActor)
+	{
+		ViewportClient->SetViewLocation(CameraActor->GetActorLocation());
+		ViewportClient->SetViewRotation(CameraActor->GetActorRotation());
+		//ViewportClient->bEditorCameraCut = !ViewportClient->IsLockedToActor(CameraActor);
+		ViewportClient->bLockedCameraView = true;
+	}
+	else
+	{
+		ViewportClient->ViewFOV = ViewportClient->FOVAngle;
+		//ViewportClient->bEditorCameraCut = false;
+		ViewportClient->bLockedCameraView = false;
+	}
+
+	// Set the actor lock.
+	ViewportClient->SetActorLock(CameraActor);
+	ViewportClient->RemoveCameraRoll();
+
+	// If viewing through a camera - enforce aspect ratio.
+	if (CameraActor)
+	{
+		if (CameraActor->GetCameraComponent()->AspectRatio == 0)
+		{
+			ViewportClient->AspectRatio = 1.7f;
+		}
+		else
+		{
+			ViewportClient->AspectRatio = CameraActor->GetCameraComponent()->AspectRatio;
+		}
+
+		//don't stop the camera from zooming when not playing back
+		ViewportClient->ViewFOV = CameraActor->GetCameraComponent()->FieldOfView;
+	}
+
+	// Update ControllingActorViewInfo, so it is in sync with the updated viewport
+	ViewportClient->UpdateViewForLockedActor();
+}
+
+FReply SCinematicLevelViewport::OnKeyDown( const FGeometry& MyGeometry, const FKeyEvent& InKeyEvent ) 
+{
+	ISequencer* Sequencer = GetSequencer();
+	if (Sequencer && Sequencer->GetCommandBindings()->ProcessCommandBindings(InKeyEvent))
+	{
+		return FReply::Handled();
+	}
+
+	return FReply::Unhandled();
+}
+
+void SCinematicLevelViewport::Setup(FLevelSequenceEditorToolkit& NewToolkit)
+{
+	CurrentToolkit = StaticCastSharedRef<FLevelSequenceEditorToolkit>(NewToolkit.AsShared());
+
+	NewToolkit.OnClosed().AddSP(this, &SCinematicLevelViewport::OnEditorClosed);
+
+	ISequencer* Sequencer = GetSequencer();
+	if (Sequencer)
+	{
+		TypeInterfaceProxy->Impl = Sequencer->GetZeroPadNumericTypeInterface();
+
+		OnCameraCutHandle = Sequencer->OnCameraCut().AddSP(this, &SCinematicLevelViewport::OnUpdateCameraCut);
+
+		if (TransportRange.IsValid())
+		{
+			TransportRange->SetSequencer(Sequencer->AsShared());
+		}
+
+		if (TransportControls.IsValid())
+		{
+			TransportControls->SetContent(Sequencer->MakeTransportControls());
+		}
+
+		SetCameraMode(ECinematicCameraMode::Auto);
+	}
+	else
+	{
+		SetCameraMode(ECinematicCameraMode::FreeCamera);
+	}
+}
+
+void SCinematicLevelViewport::CleanUp()
+{
+	ISequencer* Sequencer = GetSequencer();
+	if (Sequencer && OnCameraCutHandle.IsValid())
+	{
+		Sequencer->OnCameraCut().Remove(OnCameraCutHandle);
+	}
+
+	TransportControls->SetContent(SNullWidget::NullWidget);
+}
+
+void SCinematicLevelViewport::OnEditorOpened(FLevelSequenceEditorToolkit& Toolkit)
+{
+	if (!CurrentToolkit.IsValid())
+	{
+		Setup(Toolkit);
+	}
+}
+
+void SCinematicLevelViewport::OnEditorClosed()
+{
+	CleanUp();
+
+	FLevelSequenceEditorToolkit* NewToolkit = nullptr;
+	FLevelSequenceEditorToolkit::IterateOpenToolkits([&](FLevelSequenceEditorToolkit& Toolkit){
+		NewToolkit = &Toolkit;
+		return false;
+	});
+
+	if (NewToolkit)
+	{
+		Setup(*NewToolkit);
+	}
+	else
+	{
+		SetCameraMode(ECinematicCameraMode::FreeCamera);
+	}
+}
+
+ISequencer* SCinematicLevelViewport::GetSequencer() const
+{
+	TSharedPtr<FLevelSequenceEditorToolkit> Toolkit = CurrentToolkit.Pin();
+	if (Toolkit.IsValid())
+	{
+		return Toolkit->GetSequencer().Get();
+	}
+
+	return nullptr;
+}
+
+void SCinematicLevelViewport::Tick(const FGeometry& AllottedGeometry, const double InCurrentTime, const float InDeltaTime)
+{
+	SCompoundWidget::Tick(AllottedGeometry, InCurrentTime, InDeltaTime);
+
+	ISequencer* Sequencer = GetSequencer();
+	UMovieSceneSequence* Sequence = Sequencer ? Sequencer->GetFocusedMovieSceneSequence() : nullptr;
+
+	if (!Sequence)
+	{
+		return;
+	}
+
+	UpdateActiveCameras();
+
+	TSharedPtr<FCinematicViewportClient> ViewportClient = GetViewportClient();
+
+	// Update the specific locked camera if set. We do this each tick as it might be a spawnable camera
+	if (LockedCameraGuid.IsSet())
+	{
+		FActiveCamera* Camera = Cameras.Find(LockedCameraGuid.GetValue());
+		if (Camera && Camera->CameraActor.IsValid())
+		{
+			ViewportClient->SetActorLock(Camera->CameraActor.Get());
+
+			ViewportClient->bLockedCameraView = true;
+			ViewportClient->bDisableInput = true;
+		}
+		else
+		{
+			// If the camera's not available, just use free cam until it is
+			ViewportClient->SetActorLock(nullptr);
+			ViewportClient->bLockedCameraView = false;
+			ViewportClient->bDisableInput = false;
+		}
+	}
+
+	// Find either a cinematic shot track or a sub track
+	UMovieSceneSubTrack* SubTrack = Cast<UMovieSceneSubTrack>(Sequence->GetMovieScene()->FindMasterTrack(UMovieSceneCinematicShotTrack::StaticClass()));
+	if (!SubTrack)
+	{
+		SubTrack = Cast<UMovieSceneSubTrack>(Sequence->GetMovieScene()->FindMasterTrack(UMovieSceneSubTrack::StaticClass()));
+	}
+
+	UMovieSceneSubSection* SubSection = nullptr;
+	if (SubTrack)
+	{
+		const float CurrentTime = Sequencer->GetGlobalTime();
+
+		for (UMovieSceneSection* Section : SubTrack->GetAllSections())
+		{
+			if (Section->IsInfinite() || Section->IsTimeWithinSection(CurrentTime))
+			{
+				SubSection = CastChecked<UMovieSceneSubSection>(Section);
+			}
+		}
+	}
+
+	const float AbsoluteTime = Sequencer->GetGlobalTime();
+
+	FText TimeFormat = LOCTEXT("TimeFormat", "{0}");
+
+	TSharedPtr<INumericTypeInterface<float>> ZeroPadTypeInterface = Sequencer->GetZeroPadNumericTypeInterface();
+
+	if (SubSection)
+	{
+		float PlaybackRangeStart = 0.f;
+		if (SubSection->GetSequence() != nullptr)
+		{
+			UMovieScene* ShotMovieScene = SubSection->GetSequence()->GetMovieScene();
+			PlaybackRangeStart = ShotMovieScene->GetPlaybackRange().GetLowerBoundValue();
+		}
+
+		const float ShotOffset = SubSection->StartOffset + PlaybackRangeStart - SubSection->PrerollTime;
+		const float AbsoluteShotPosition = ShotOffset + (AbsoluteTime - (SubSection->GetStartTime() - SubSection->PrerollTime)) / SubSection->TimeScale;
+		UIData.Frame = FText::Format(
+			TimeFormat,
+			FText::FromString(ZeroPadTypeInterface->ToString(AbsoluteShotPosition))
+		);
+
+		UMovieSceneCinematicShotSection* CinematicShotSection = Cast<UMovieSceneCinematicShotSection>(SubSection);
+		if (CinematicShotSection)
+		{
+			UIData.ShotName = CinematicShotSection->GetShotDisplayName();
+		}
+		else
+		{
+			UIData.ShotName = SubSection->GetSequence()->GetDisplayName();
+		}
+	}
+	else
+	{
+		UIData.Frame = FText::Format(
+			TimeFormat,
+			FText::FromString(ZeroPadTypeInterface->ToString(AbsoluteTime))
+			);
+
+		UIData.ShotName = Sequence->GetDisplayName();
+	}
+
+	TRange<float> EntireRange = Sequence->GetMovieScene()->GetEditorData().WorkingRange;
+
+	UIData.MasterStartText = FText::Format(
+		TimeFormat,
+		FText::FromString(ZeroPadTypeInterface->ToString(EntireRange.GetLowerBoundValue()))
+	);
+
+	UIData.MasterEndText = FText::Format(
+		TimeFormat,
+		FText::FromString(ZeroPadTypeInterface->ToString(EntireRange.GetUpperBoundValue()))
+	);
+
+	UCameraComponent* CameraComponent = ViewportClient->GetCameraComponentForView();
+	if (CameraComponent)
+	{
+		if (UCineCameraComponent* CineCam = Cast<UCineCameraComponent>(CameraComponent))
+		{
+			const float SensorWidth = CineCam->FilmbackSettings.SensorWidth;
+			const float SensorHeight = CineCam->FilmbackSettings.SensorHeight;
+
+			// Search presets for one that matches
+			const FNamedFilmbackPreset* Preset = UCineCameraComponent::GetFilmbackPresets().FindByPredicate([&](const FNamedFilmbackPreset& InPreset){
+				return InPreset.FilmbackSettings.SensorWidth == SensorWidth && InPreset.FilmbackSettings.SensorHeight == SensorHeight;
+			});
+
+			if (Preset)
+			{
+				UIData.Filmback = FText::FromString(Preset->Name);
+			}
+			else
+			{
+				FNumberFormattingOptions Opts = FNumberFormattingOptions().SetMaximumFractionalDigits(1);
+				UIData.Filmback = FText::Format(
+					LOCTEXT("CustomFilmbackFormat", "Custom ({0}mm x {1}mm)"),
+					FText::AsNumber(SensorWidth, &Opts),
+					FText::AsNumber(SensorHeight, &Opts)
+				);
+			}
+		}
+		else
+		{
+			// Just use the FoV
+			UIData.Filmback = FText::FromString(LexicalConversion::ToString(FNumericUnit<float>(CameraComponent->FieldOfView, EUnit::Degrees)));
+		}
+	}
+	else
+	{
+		UIData.Filmback = FText();
+	}
+}
+
+#undef LOCTEXT_NAMESPACE
\ No newline at end of file
diff --git a/Engine/Plugins/MovieScene/LevelSequenceEditor/Source/LevelSequenceEditor/Private/CinematicViewport/SCinematicLevelViewport.h b/Engine/Plugins/MovieScene/LevelSequenceEditor/Source/LevelSequenceEditor/Private/CinematicViewport/SCinematicLevelViewport.h
new file mode 100644
--- /dev/null
+++ b/Engine/Plugins/MovieScene/LevelSequenceEditor/Source/LevelSequenceEditor/Private/CinematicViewport/SCinematicLevelViewport.h
@@ -0,0 +1,197 @@
+// Copyright 1998-2015 Epic Games, Inc. All Rights Reserved.
+
+#pragma once
+
+#include "SEditorViewport.h"
+
+class FLevelSequenceEditorToolkit;
+class ISequencer;
+class SCinematicTransportRange;
+class FLevelViewportLayout;
+struct FTypeInterfaceProxy;
+
+/** Overridden level viewport client for this viewport */
+struct FCinematicViewportClient : FLevelEditorViewportClient
+{
+	FCinematicViewportClient();
+	virtual bool IsLevelEditorClient() const override { return true; }
+};
+
+/** struct containing UI data - populated once per frame */
+struct FUIData
+{
+	/** The name of the current shot */
+	FText ShotName;
+	/** The name of the current shot's lens */
+	FText Lens;
+	/** The name of the current shot's filmback */
+	FText Filmback;
+	/** The text that represents the current frame */
+	FText Frame;
+	/** The text that represents the master start frame */
+	FText MasterStartText;
+	/** The text that represents the master end frame */
+	FText MasterEndText;
+};
+
+/** Enum specifying different camera modes */
+enum class ECinematicCameraMode
+{
+	FreeCamera,
+	Auto,
+	SpecificCamera
+};
+
+class SCinematicLevelViewport : public SCompoundWidget
+{
+public:
+
+	SLATE_BEGIN_ARGS(SCinematicLevelViewport) : _bShowMaximize(true) {}
+		
+		/** The unique name of this viewport inside its parent layout */
+		SLATE_ARGUMENT(FName, LayoutName)
+		/** Name of the viewport layout we should revert to at the user's request */
+		SLATE_ARGUMENT(FName, RevertToLayoutName)
+		/** Ptr to this viewport's parent layout */
+		SLATE_ARGUMENT(TSharedPtr<FLevelViewportLayout>, ParentLayout)
+		/** Whether to show the maximize button or not */
+		SLATE_ARGUMENT(bool, bShowMaximize)
+	SLATE_END_ARGS()
+
+	/** Access this viewport's viewport client */
+	TSharedPtr<FCinematicViewportClient> GetViewportClient() const;
+
+	/** Construct this widget */
+	void Construct(const FArguments& InArgs);
+
+	virtual void Tick(const FGeometry& AllottedGeometry, const double InCurrentTime, const float InDeltaTime);
+	virtual FReply OnKeyDown(const FGeometry& MyGeometry, const FKeyEvent& InKeyEvent) override;
+	virtual bool SupportsKeyboardFocus() const override { return true; }
+private:
+
+	/** Set up this viewport to operate on the specified toolkit */
+	void Setup(FLevelSequenceEditorToolkit& NewToolkit);
+
+	/** Clean up this viewport after its toolkit has been closed */
+	void CleanUp();
+
+	/** Called when a level sequence editor toolkit has been opened */
+	void OnEditorOpened(FLevelSequenceEditorToolkit& Toolkit);
+	
+	/** Called when the level sequence editor toolkit we were observing has been closed */
+	void OnEditorClosed();
+
+private:
+
+	/** Get the sequencer ptr from our current toolkit, if set */
+	ISequencer* GetSequencer() const;
+
+	/** Called when our sequencer wants to switch cameras */
+	void OnUpdateCameraCut(UObject* CameraObject);
+
+	/** Set the camera to the specified mode (optionally binding to a specific camera) */
+	void SetCameraMode(ECinematicCameraMode InCameraMode, TOptional<FGuid> InSpecificCameraGuid = TOptional<FGuid>());
+
+	/** Make the currently available camera menu */
+	TSharedRef<SWidget> MakeCameraMenu();
+
+private:
+
+	/** Toggle this viewport's maximized state */
+	FReply OnToggleMaximize();
+
+	/** Test whether this viewport is maximized or not */
+	bool IsMaximized() const;
+
+	/** Called when we want to revert to the equivalent non-cinematic viewport */
+	FReply OnRevertLayout();
+
+private:
+
+	/** Called once per frame to update the available cameras in the scene */
+	void UpdateActiveCameras();
+
+	/** Cache the desired viewport size based on the filled, allotted geometry */
+	void CacheDesiredViewportSize(const FGeometry& AllottedGeometry);
+
+	/** Get the desired width and height of the viewport */
+	FOptionalSize GetDesiredViewportWidth() const;
+	FOptionalSize GetDesiredViewportHeight() const;
+
+	/** Get the image for the maximize button */
+	const FSlateBrush* GetMaximizeImage() const;
+
+	int32 GetVisibleWidgetIndex() const;
+	EVisibility GetControlsVisibility() const;
+
+private:
+
+	TOptional<float> GetMinTime() const;
+
+	TOptional<float> GetMaxTime() const;
+
+	void OnTimeCommitted(float Value, ETextCommit::Type);
+
+	void SetTime(float Value);
+
+	float GetTime() const;
+
+	FReply SetPlaybackStart();
+
+	FReply SetPlaybackEnd();
+
+private:
+	
+	/** Widget where the scene viewport is drawn in */
+	TSharedPtr<SEditorViewport> ViewportWidget;
+
+	/** The toolkit we're currently editing */
+	TWeakPtr<FLevelSequenceEditorToolkit> CurrentToolkit;
+
+	/** Slot for transport controls */
+	TSharedPtr<SBox> TransportControls;
+
+	/** Transport range widget */
+	TSharedPtr<SCinematicTransportRange> TransportRange;
+	
+	/** Cached UI data */
+	FUIData UIData;
+
+	/** Delegate binding handle for ISequencer::OnCameraCut */
+	FDelegateHandle OnCameraCutHandle;
+
+	/** The current camera mode */
+	ECinematicCameraMode CameraMode;
+
+	/** Guid of the camera in the UMovieScene to look through. Stored as a GUID rather than a ptr to ensure spawnable cameras work. */
+	TOptional<FGuid> LockedCameraGuid;
+
+	/** Cached desired size of the viewport */
+	FVector2D DesiredViewportSize;
+	
+	/** Viewport controls widget that sits just below, and matches the width, of the viewport */
+	TSharedPtr<SWidget> ViewportControls;
+
+	/** Weak ptr to our parent layout */
+	TWeakPtr<FLevelViewportLayout> ParentLayout;
+
+	/** Name of the viewport in the parent layout */
+	FName LayoutName;
+	/** Name of the viewport layout we should revert to at the user's request */
+	FName RevertToLayoutName;
+
+	/** The time spin box */
+	TSharedPtr<FTypeInterfaceProxy> TypeInterfaceProxy;
+
+	/** Cached active cameras */
+	struct FActiveCamera
+	{
+		FActiveCamera(FString InDisplayName, ACameraActor* InCameraActor)
+			: DisplayName(InDisplayName), CameraActor(InCameraActor)
+		{}
+
+		FString DisplayName;
+		TWeakObjectPtr<ACameraActor> CameraActor;
+	};
+	TMap<FGuid, FActiveCamera> Cameras;
+};
\ No newline at end of file
diff --git a/Engine/Plugins/MovieScene/LevelSequenceEditor/Source/LevelSequenceEditor/Private/CinematicViewport/SCinematicTransportRange.cpp b/Engine/Plugins/MovieScene/LevelSequenceEditor/Source/LevelSequenceEditor/Private/CinematicViewport/SCinematicTransportRange.cpp
new file mode 100644
--- /dev/null
+++ b/Engine/Plugins/MovieScene/LevelSequenceEditor/Source/LevelSequenceEditor/Private/CinematicViewport/SCinematicTransportRange.cpp
@@ -0,0 +1,199 @@
+// Copyright 1998-2015 Epic Games, Inc. All Rights Reserved.
+
+#include "LevelSequenceEditorPCH.h"
+#include "LevelSequenceEditorStyle.h"
+#include "SCinematicTransportRange.h"
+#include "MovieScene.h"
+#include "ISequencer.h"
+
+#define LOCTEXT_NAMESPACE "SCinematicTransportRange"
+
+
+void SCinematicTransportRange::Construct(const FArguments& InArgs)
+{
+	bDraggingTime = false;
+}
+
+void SCinematicTransportRange::SetSequencer(TWeakPtr<ISequencer> InSequencer)
+{
+	WeakSequencer = InSequencer;
+}
+
+ISequencer* SCinematicTransportRange::GetSequencer() const
+{
+	return WeakSequencer.Pin().Get();
+}
+
+FVector2D SCinematicTransportRange::ComputeDesiredSize(float) const
+{
+	static const float MarkerHeight = 6.f;
+	static const float TrackHeight = 8.f;
+	return FVector2D(100.f, MarkerHeight + TrackHeight);
+}
+
+FReply SCinematicTransportRange::OnMouseButtonDown(const FGeometry& MyGeometry, const FPointerEvent& MouseEvent)
+{
+	bDraggingTime = true;
+	SetTime(MyGeometry, MouseEvent);
+	return FReply::Handled().CaptureMouse(AsShared()).PreventThrottling();
+}
+
+FReply SCinematicTransportRange::OnMouseMove(const FGeometry& MyGeometry, const FPointerEvent& MouseEvent)
+{
+	if (bDraggingTime)
+	{
+		SetTime(MyGeometry, MouseEvent);
+	}
+
+	return FReply::Handled();
+}
+
+FReply SCinematicTransportRange::OnMouseButtonUp(const FGeometry& MyGeometry, const FPointerEvent& MouseEvent)
+{
+	bDraggingTime = false;
+	return FReply::Handled().ReleaseMouseCapture();
+}
+
+void SCinematicTransportRange::SetTime(const FGeometry& MyGeometry, const FPointerEvent& MouseEvent)
+{
+	ISequencer* Sequencer = GetSequencer();
+	if (Sequencer)
+	{
+		float Lerp = MyGeometry.AbsoluteToLocal(MouseEvent.GetScreenSpacePosition()).X / MyGeometry.Size.X;
+		Lerp = FMath::Clamp(Lerp, 0.f, 1.f);
+		
+		const TRange<float> WorkingRange = Sequencer->GetFocusedMovieSceneSequence()->GetMovieScene()->GetEditorData().WorkingRange;
+		const bool bAllowSnappingToFrames = true;
+		Sequencer->SetGlobalTime(WorkingRange.GetLowerBoundValue() + WorkingRange.Size<float>()*Lerp, bAllowSnappingToFrames);
+	}
+}
+
+void SCinematicTransportRange::OnMouseCaptureLost()
+{
+	bDraggingTime = false;
+}
+
+int32 SCinematicTransportRange::OnPaint(const FPaintArgs& Args, const FGeometry& AllottedGeometry, const FSlateRect& MyClippingRect, FSlateWindowElementList& OutDrawElements, int32 LayerId, const FWidgetStyle& InWidgetStyle, bool bParentEnabled) const
+{
+	ISequencer* Sequencer = GetSequencer();
+
+	ESlateDrawEffect::Type DrawEffects = bParentEnabled ? ESlateDrawEffect::None : ESlateDrawEffect::DisabledEffect;
+	if (!Sequencer)
+	{
+		return LayerId;
+	}
+
+	static const float TrackOffsetY = 6.f;
+	const float TrackHeight = AllottedGeometry.Size.Y - TrackOffsetY;
+
+	TRange<float> WorkingRange = Sequencer->GetFocusedMovieSceneSequence()->GetMovieScene()->GetEditorData().WorkingRange;
+	TRange<float> PlaybackRange = Sequencer->GetFocusedMovieSceneSequence()->GetMovieScene()->GetPlaybackRange();
+
+	FColor DarkGray(40, 40, 40);
+	FColor MidGray(80, 80, 80);
+	FColor LightGray(200, 200, 200);
+
+	// Paint the left padding before the playback start
+	FSlateDrawElement::MakeBox(
+		OutDrawElements,
+		LayerId,
+		AllottedGeometry.ToPaintGeometry( FVector2D(0.f, TrackOffsetY),  FVector2D(AllottedGeometry.Size.X, TrackHeight)),
+		FEditorStyle::GetBrush("WhiteBrush"),
+		MyClippingRect,
+		DrawEffects,
+		FLinearColor(DarkGray)
+	);
+	
+	const float FullRange = WorkingRange.Size<float>();
+
+	const float PlaybackStartLerp	= (PlaybackRange.GetLowerBoundValue() - WorkingRange.GetLowerBoundValue()) / FullRange;
+	const float PlaybackEndLerp		= (PlaybackRange.GetUpperBoundValue() - WorkingRange.GetLowerBoundValue()) / FullRange;
+
+	// Draw the playback range
+	FSlateDrawElement::MakeBox(
+		OutDrawElements,
+		++LayerId,
+		AllottedGeometry.ToPaintGeometry( FVector2D(AllottedGeometry.Size.X*PlaybackStartLerp, TrackOffsetY),  FVector2D(AllottedGeometry.Size.X*(PlaybackEndLerp - PlaybackStartLerp), TrackHeight)),
+		FEditorStyle::GetBrush("WhiteBrush"),
+		MyClippingRect,
+		DrawEffects,
+		FLinearColor(MidGray)
+	);
+
+	const float ProgressLerp = (Sequencer->GetGlobalTime() - WorkingRange.GetLowerBoundValue()) / FullRange;
+
+	// Draw the playback progress
+	if (ProgressLerp > PlaybackStartLerp)
+	{
+		const float ClampedProgressLerp = FMath::Clamp(ProgressLerp, PlaybackStartLerp, PlaybackEndLerp);
+		FSlateDrawElement::MakeBox(
+			OutDrawElements,
+			++LayerId,
+			AllottedGeometry.ToPaintGeometry( FVector2D(AllottedGeometry.Size.X*PlaybackStartLerp, TrackOffsetY), FVector2D(AllottedGeometry.Size.X * (ClampedProgressLerp - PlaybackStartLerp), TrackHeight) ),
+			FEditorStyle::GetBrush("WhiteBrush"),
+			MyClippingRect,
+			DrawEffects,
+			FLinearColor(LightGray)
+		);
+	}
+
+	// Draw the play marker
+	{
+		static const float BrushWidth = 11.f, BrushHeight = 6.f;
+		const float PositionX = AllottedGeometry.Size.X * ProgressLerp;
+
+		FSlateDrawElement::MakeBox(
+			OutDrawElements,
+			++LayerId,
+			AllottedGeometry.ToPaintGeometry(FVector2D(PositionX - FMath::CeilToFloat(BrushWidth/2), 0.f), FVector2D(BrushWidth, BrushHeight)),
+			FLevelSequenceEditorStyle::Get()->GetBrush("LevelSequenceEditor.CinematicViewportPlayMarker"),
+			MyClippingRect,
+			DrawEffects,
+			LightGray
+		);
+
+		TArray<FVector2D> LinePoints;
+		LinePoints.Add(FVector2D(PositionX, BrushHeight));
+		LinePoints.Add(FVector2D(PositionX, AllottedGeometry.Size.Y));
+
+		FSlateDrawElement::MakeLines(
+			OutDrawElements,
+			LayerId,
+			AllottedGeometry.ToPaintGeometry(),
+			LinePoints,
+			MyClippingRect,
+			DrawEffects,
+			LightGray,
+			false
+		);
+	}
+
+	// Draw the play bounds
+	{
+		static const float BrushWidth = 4.f;
+
+		FSlateDrawElement::MakeBox(
+			OutDrawElements,
+			LayerId+1,
+			AllottedGeometry.ToPaintGeometry(FVector2D(AllottedGeometry.Size.X*PlaybackStartLerp, TrackOffsetY), FVector2D(BrushWidth, TrackHeight)),
+			FLevelSequenceEditorStyle::Get()->GetBrush("LevelSequenceEditor.CinematicViewportRangeStart"),
+			MyClippingRect,
+			DrawEffects,
+			FColor(32, 128, 32)	// 120, 75, 50 (HSV)
+		);
+
+		FSlateDrawElement::MakeBox(
+			OutDrawElements,
+			LayerId+1,
+			AllottedGeometry.ToPaintGeometry(FVector2D(AllottedGeometry.Size.X*PlaybackEndLerp - BrushWidth, TrackOffsetY), FVector2D(BrushWidth, TrackHeight)),
+			FLevelSequenceEditorStyle::Get()->GetBrush("LevelSequenceEditor.CinematicViewportRangeEnd"),
+			MyClippingRect,
+			DrawEffects,
+			FColor(128, 32, 32)	// 0, 75, 50 (HSV)
+		);
+	}
+
+	return LayerId;
+}
+
+#undef LOCTEXT_NAMESPACE
\ No newline at end of file
diff --git a/Engine/Plugins/MovieScene/LevelSequenceEditor/Source/LevelSequenceEditor/Private/Customizations/LevelSequencePlaybackSettingsCustomization.cpp b/Engine/Plugins/MovieScene/LevelSequenceEditor/Source/LevelSequenceEditor/Private/Customizations/LevelSequencePlaybackSettingsCustomization.cpp
--- a/Engine/Plugins/MovieScene/LevelSequenceEditor/Source/LevelSequenceEditor/Private/Customizations/LevelSequencePlaybackSettingsCustomization.cpp
+++ b/Engine/Plugins/MovieScene/LevelSequenceEditor/Source/LevelSequenceEditor/Private/Customizations/LevelSequencePlaybackSettingsCustomization.cpp
@@ -105,6 +105,23 @@ void FLevelSequencePlaybackSettingsCustomization::CustomizeChildren(TSharedRef<I
 								.Text_Raw(this, &FLevelSequencePlaybackSettingsCustomization::GetCustomSuffix)
 						]
 				]
+
+			+SHorizontalBox::Slot()
+				.VAlign(VAlign_Center)
+				.AutoWidth()
+				[
+					SNew(SButton)
+					.OnClicked(this, &FLevelSequencePlaybackSettingsCustomization::OnLoopResetClicked)
+					.Visibility(this, &FLevelSequencePlaybackSettingsCustomization::GetLoopResetVisibility)
+					.ContentPadding(FMargin(5.f, 0.f))
+					.ToolTipText(LOCTEXT("ResetToDefaultToolTip", "Reset to Default"))
+					.ButtonStyle( FEditorStyle::Get(), "NoBorder" )
+					.Content()
+					[
+						SNew(SImage)
+						.Image( FEditorStyle::GetBrush("PropertyWindow.DiffersFromDefault") )
+					]
+				]
 		];
 }
 
@@ -119,6 +136,19 @@ FText FLevelSequencePlaybackSettingsCustomization::GetCustomSuffix() const
 		: LOCTEXT("Times", "times");
 }
 
+EVisibility FLevelSequencePlaybackSettingsCustomization::GetLoopResetVisibility() const
+{
+	return CurrentMode == LoopModes[0] ? EVisibility::Hidden : EVisibility::Visible;
+}
+
+FReply FLevelSequencePlaybackSettingsCustomization::OnLoopResetClicked()
+{
+	CurrentMode = LoopModes[0];
+
+	UpdateProperty();
+
+	return FReply::Handled();
+}
 
 void FLevelSequencePlaybackSettingsCustomization::UpdateProperty()
 {
diff --git a/Engine/Plugins/MovieScene/LevelSequenceEditor/Source/LevelSequenceEditor/Private/Customizations/LevelSequencePlaybackSettingsCustomization.h b/Engine/Plugins/MovieScene/LevelSequenceEditor/Source/LevelSequenceEditor/Private/Customizations/LevelSequencePlaybackSettingsCustomization.h
--- a/Engine/Plugins/MovieScene/LevelSequenceEditor/Source/LevelSequenceEditor/Private/Customizations/LevelSequencePlaybackSettingsCustomization.h
+++ b/Engine/Plugins/MovieScene/LevelSequenceEditor/Source/LevelSequenceEditor/Private/Customizations/LevelSequencePlaybackSettingsCustomization.h
@@ -31,6 +31,12 @@ private:
 	/** Get the suffix to display after the custom entry box. */
 	FText GetCustomSuffix() const;
 
+	/** Get whether the loop reset button should be visible. */
+	EVisibility GetLoopResetVisibility() const;
+
+	/** Called when the "Reset to Default" button for the loop mode has been clicked. */
+	FReply OnLoopResetClicked();
+
 	/** Array of loop modes. */
 	struct FLoopMode
 	{
diff --git a/Engine/Plugins/MovieScene/LevelSequenceEditor/Source/LevelSequenceEditor/Private/Factories/LevelSequenceFactoryNew.cpp b/Engine/Plugins/MovieScene/LevelSequenceEditor/Source/LevelSequenceEditor/Private/Factories/LevelSequenceFactoryNew.cpp
--- a/Engine/Plugins/MovieScene/LevelSequenceEditor/Source/LevelSequenceEditor/Private/Factories/LevelSequenceFactoryNew.cpp
+++ b/Engine/Plugins/MovieScene/LevelSequenceEditor/Source/LevelSequenceEditor/Private/Factories/LevelSequenceFactoryNew.cpp
@@ -2,9 +2,9 @@
 
 #include "LevelSequenceEditorPCH.h"
 #include "LevelSequence.h"
-#include "LevelSequenceProjectSettings.h"
 #include "LevelSequenceFactoryNew.h"
 #include "MovieScene.h"
+#include "MovieSceneToolsProjectSettings.h"
 
 #define LOCTEXT_NAMESPACE "MovieSceneFactory"
 
@@ -30,7 +30,7 @@ UObject* ULevelSequenceFactoryNew::FactoryCreateNew(UClass* Class, UObject* InPa
 	NewLevelSequence->Initialize();
 	
 	// Set up some sensible defaults
-	const ULevelSequenceProjectSettings* ProjectSettings = GetDefault<ULevelSequenceProjectSettings>();
+	const UMovieSceneToolsProjectSettings* ProjectSettings = GetDefault<UMovieSceneToolsProjectSettings>();
 	NewLevelSequence->GetMovieScene()->SetPlaybackRange(ProjectSettings->DefaultStartTime, ProjectSettings->DefaultStartTime + ProjectSettings->DefaultDuration);
 
 	return NewLevelSequence;
diff --git a/Engine/Plugins/MovieScene/LevelSequenceEditor/Source/LevelSequenceEditor/Private/LevelSequenceEditorModule.cpp b/Engine/Plugins/MovieScene/LevelSequenceEditor/Source/LevelSequenceEditor/Private/LevelSequenceEditorModule.cpp
--- a/Engine/Plugins/MovieScene/LevelSequenceEditor/Source/LevelSequenceEditor/Private/LevelSequenceEditorModule.cpp
+++ b/Engine/Plugins/MovieScene/LevelSequenceEditor/Source/LevelSequenceEditor/Private/LevelSequenceEditorModule.cpp
@@ -6,11 +6,13 @@
 #include "LevelSequenceEditorStyle.h"
 #include "ModuleInterface.h"
 #include "PropertyEditorModule.h"
-#include "ISettingsModule.h"
-#include "LevelSequenceProjectSettings.h"
+#include "CinematicLevelViewportLayout.h"
+#include "IPlacementModeModule.h"
+#include "Camera/CineCameraActor.h"
 
 #define LOCTEXT_NAMESPACE "LevelSequenceEditor"
 
+TSharedPtr<FLevelSequenceEditorStyle> FLevelSequenceEditorStyle::Singleton;
 
 class FLevelSequenceExtensionCommands
 	: public TCommands<FLevelSequenceExtensionCommands>
@@ -44,41 +46,61 @@ public:
 
 	virtual void StartupModule() override
 	{
-		Style = MakeShareable(new FLevelSequenceEditorStyle());
-
-		if (ISettingsModule* SettingsModule = FModuleManager::GetModulePtr<ISettingsModule>("Settings"))
-		{
-			SettingsModule->RegisterSettings("Project", "Editor", "Level Sequences",
-				LOCTEXT("RuntimeSettingsName", "Level Sequences"),
-				LOCTEXT("RuntimeSettingsDescription", "Configure project settings relating to Level Sequences"),
-				GetMutableDefault<ULevelSequenceProjectSettings>()
-			);
-		}
+		FLevelSequenceEditorStyle::Get();
 
 		RegisterAssetTools();
 		RegisterCustomizations();
 		RegisterMenuExtensions();
+		RegisterLevelEditorExtensions();
+		RegisterPlacementModeExtensions();
 	}
 	
 	virtual void ShutdownModule() override
 	{
-		if (ISettingsModule* SettingsModule = FModuleManager::GetModulePtr<ISettingsModule>("Settings"))
-		{
-			SettingsModule->UnregisterSettings("Project", "Editor", "Level Sequences");
-		}
-
 		UnregisterAssetTools();
 		UnregisterCustomizations();
 		UnregisterMenuExtensions();
+		UnregisterLevelEditorExtensions();
+		UnregisterPlacementModeExtensions();
 	}
 
 protected:
 
+	/** Registers level editor extensions. */
+	void RegisterLevelEditorExtensions()
+	{
+		FLevelEditorModule& LevelEditorModule = FModuleManager::LoadModuleChecked<FLevelEditorModule>("LevelEditor");
+
+		FCustomViewportLayoutDefinition TwoPaneDefinition = FCustomViewportLayoutDefinition::FromType<FCinematicLevelViewportLayout_OnePane>();
+		TwoPaneDefinition.DisplayName = LOCTEXT("OnePaneCinematicLayoutName", "Cinematic");
+		TwoPaneDefinition.Description = LOCTEXT("OnePaneCinematicLayoutDesc", "A viewport layout tailored to cinematic preview");
+		TwoPaneDefinition.Icon = FSlateIcon("LevelSequenceEditorStyle", "LevelSequenceEditor.OnePaneCinematicViewportLayout");
+		LevelEditorModule.RegisterCustomViewportLayout("OnePaneCinematic", TwoPaneDefinition);
+
+		FCustomViewportLayoutDefinition OnePaneDefinition = FCustomViewportLayoutDefinition::FromType<FCinematicLevelViewportLayout_TwoPane>();
+		OnePaneDefinition.DisplayName = LOCTEXT("TwoPaneCinematicLayoutName", "Two Pane Cinematic");
+		OnePaneDefinition.Description = LOCTEXT("TwoPaneCinematicLayoutDesc", "A viewport layout comprising an edit viewport, and a cinematic preview viewport");
+		OnePaneDefinition.Icon = FSlateIcon("LevelSequenceEditorStyle", "LevelSequenceEditor.TwoPaneCinematicViewportLayout");
+		LevelEditorModule.RegisterCustomViewportLayout("TwoPaneCinematic", OnePaneDefinition);
+	}
+
 	/** Registers asset tool actions. */
 	void RegisterAssetTools()
 	{
 		IAssetTools& AssetTools = FModuleManager::LoadModuleChecked<FAssetToolsModule>("AssetTools").Get();
-		RegisterAssetTypeAction(AssetTools, MakeShareable(new FLevelSequenceActions(Style.ToSharedRef())));
+		RegisterAssetTypeAction(AssetTools, MakeShareable(new FLevelSequenceActions(FLevelSequenceEditorStyle::Get())));
+	}
+
+	/** Registers placement mode extensions. */
+	void RegisterPlacementModeExtensions()
+	{
+		FPlacementCategoryInfo Info(
+			LOCTEXT("CinematicCategoryName", "Cinematic"),
+			"Cinematic",
+			TEXT("PMCinematic"),
+			25);
+		IPlacementModeModule::Get().RegisterPlacementCategory(Info);
+		IPlacementModeModule::Get().RegisterPlaceableItem(Info.UniqueHandle, MakeShareable( new FPlaceableItem(nullptr, FAssetData(ACineCameraActor::StaticClass())) ));
 	}
 
 	/**
@@ -109,6 +131,25 @@ protected:
 		}
 	}
 
+	/** Unregisters level editor extensions. */
+	void UnregisterLevelEditorExtensions()
+	{
+		FLevelEditorModule* LevelEditorModule = FModuleManager::GetModulePtr<FLevelEditorModule>("LevelEditor");
+		if (LevelEditorModule)
+		{
+			LevelEditorModule->UnRegisterCustomViewportLayout("Cinematic");
+		}
+	}
+
+	/** Unregisters placement mode extensions. */
+	void UnregisterPlacementModeExtensions()
+	{
+		if (IPlacementModeModule::IsAvailable())
+		{
+			IPlacementModeModule::Get().UnregisterPlacementCategory("Cinematic");
+		}
+	}
+
 protected:
 
 	/** Register details view customizations. */
@@ -121,8 +162,11 @@ protected:
 	/** Unregister details view customizations. */
 	void UnregisterCustomizations()
 	{
-		FPropertyEditorModule& PropertyModule = FModuleManager::LoadModuleChecked<FPropertyEditorModule>("PropertyEditor");
-		PropertyModule.UnregisterCustomPropertyTypeLayout(FLevelSequencePlaybackSettings::StaticStruct()->GetFName());
+		FPropertyEditorModule* PropertyModule = FModuleManager::GetModulePtr<FPropertyEditorModule>("PropertyEditor");
+		if (PropertyModule)
+		{
+			PropertyModule->UnregisterCustomPropertyTypeLayout(FLevelSequencePlaybackSettings::StaticStruct()->GetFName());
+		}
 	}
 
 protected:
@@ -197,7 +241,7 @@ protected:
 		}
 
 		ALevelSequenceActor* NewActor = CastChecked<ALevelSequenceActor>(GEditor->UseActorFactory(ActorFactory, FAssetData(NewAsset), &FTransform::Identity));
-		if (GCurrentLevelEditingViewportClient->IsPerspective())
+		if (GCurrentLevelEditingViewportClient != nullptr && GCurrentLevelEditingViewportClient->IsPerspective())
 		{
 			GEditor->MoveActorInFrontOfCamera(*NewActor, GCurrentLevelEditingViewportClient->GetViewLocation(), GCurrentLevelEditingViewportClient->GetViewRotation().Vector());
 		}
@@ -214,9 +258,6 @@ private:
 	/** The collection of registered asset type actions. */
 	TArray<TSharedRef<IAssetTypeActions>> RegisteredAssetTypeActions;
 
-	/** Holds the plug-ins style set. */
-	TSharedPtr<ISlateStyle> Style;
-
 	/** Extender for the cinematics menu */
 	TSharedPtr<FExtender> CinematicsMenuExtender;
 
diff --git a/Engine/Plugins/MovieScene/LevelSequenceEditor/Source/LevelSequenceEditor/Private/LevelSequenceEditorPCH.h b/Engine/Plugins/MovieScene/LevelSequenceEditor/Source/LevelSequenceEditor/Private/LevelSequenceEditorPCH.h
--- a/Engine/Plugins/MovieScene/LevelSequenceEditor/Source/LevelSequenceEditor/Private/LevelSequenceEditorPCH.h
+++ b/Engine/Plugins/MovieScene/LevelSequenceEditor/Source/LevelSequenceEditor/Private/LevelSequenceEditorPCH.h
@@ -21,4 +21,5 @@
 
 #include "LevelSequenceActions.h"
 #include "LevelSequenceEditorToolkit.h"
+#include "LevelSequenceEditorStyle.h"
 #include "LevelSequencePlaybackSettingsCustomization.h"
diff --git a/Engine/Plugins/MovieScene/LevelSequenceEditor/Source/LevelSequenceEditor/Private/LevelSequenceEditorToolkit.cpp b/Engine/Plugins/MovieScene/LevelSequenceEditor/Source/LevelSequenceEditor/Private/LevelSequenceEditorToolkit.cpp
--- a/Engine/Plugins/MovieScene/LevelSequenceEditor/Source/LevelSequenceEditor/Private/LevelSequenceEditorToolkit.cpp
+++ b/Engine/Plugins/MovieScene/LevelSequenceEditor/Source/LevelSequenceEditor/Private/LevelSequenceEditorToolkit.cpp
@@ -34,6 +34,25 @@ namespace SequencerDefs
 }
 
 
+static TArray<FLevelSequenceEditorToolkit*> OpenToolkits;
+
+void FLevelSequenceEditorToolkit::IterateOpenToolkits(TFunctionRef<bool(FLevelSequenceEditorToolkit&)> Iter)
+{
+	for (FLevelSequenceEditorToolkit* Toolkit : OpenToolkits)
+	{
+		if (!Iter(*Toolkit))
+		{
+			return;
+		}
+	}
+}
+
+FLevelSequenceEditorToolkit::FLevelSequenceEditorToolkitOpened& FLevelSequenceEditorToolkit::OnOpened()
+{
+	static FLevelSequenceEditorToolkitOpened OnOpenedEvent;
+	return OnOpenedEvent;
+}
+
 /* FLevelSequenceEditorToolkit structors
  *****************************************************************************/
 
@@ -45,6 +64,8 @@ FLevelSequenceEditorToolkit::FLevelSequenceEditorToolkit(const TSharedRef<ISlate
 	int32 NewIndex = SequencerModule.GetMenuExtensibilityManager()->GetExtenderDelegates().Add(
 		FAssetEditorExtender::CreateRaw(this, &FLevelSequenceEditorToolkit::HandleMenuExtensibilityGetExtender));
 	SequencerExtenderHandle = SequencerModule.GetMenuExtensibilityManager()->GetExtenderDelegates()[NewIndex].GetHandle();
+
+	OpenToolkits.Add(this);
 }
 
 
@@ -136,6 +157,8 @@ void FLevelSequenceEditorToolkit::Initialize(const EToolkitMode::Type Mode, cons
 		// when previewing a MovieScene
 		LevelEditorModule.OnMapChanged().AddRaw(this, &FLevelSequenceEditorToolkit::HandleMapChanged);
 	}
+
+	FLevelSequenceEditorToolkit::OnOpened().Broadcast(*this);
 }
 
 
@@ -433,5 +456,11 @@ void FLevelSequenceEditorToolkit::HandleTrackMenuExtensionAddTrack(FMenuBuilder&
 	}
 }
 
+bool FLevelSequenceEditorToolkit::OnRequestClose()
+{
+	OpenToolkits.Remove(this);
+	OnClosedEvent.Broadcast();
+	return true;
+}
 
 #undef LOCTEXT_NAMESPACE
diff --git a/Engine/Plugins/MovieScene/LevelSequenceEditor/Source/LevelSequenceEditor/Private/LevelSequenceEditorToolkit.h b/Engine/Plugins/MovieScene/LevelSequenceEditor/Source/LevelSequenceEditor/Private/LevelSequenceEditorToolkit.h
--- a/Engine/Plugins/MovieScene/LevelSequenceEditor/Source/LevelSequenceEditor/Private/LevelSequenceEditorToolkit.h
+++ b/Engine/Plugins/MovieScene/LevelSequenceEditor/Source/LevelSequenceEditor/Private/LevelSequenceEditorToolkit.h
@@ -5,7 +5,6 @@
 #include "Toolkits/AssetEditorToolkit.h"
 #include "LevelSequence.h"
 
-
 enum class EMapChangeType : uint8;
 class FTabManager;
 class ILevelViewport;
@@ -35,6 +34,19 @@ public:
 	/** Virtual destructor */
 	virtual ~FLevelSequenceEditorToolkit();
 
+public:
+
+	/** Iterate all open level sequence editor toolkits */
+	static void IterateOpenToolkits(TFunctionRef<bool(FLevelSequenceEditorToolkit&)> Iter);
+
+	/** Called when the tab manager is changed */
+	DECLARE_EVENT_OneParam(FLevelSequenceEditorToolkit, FLevelSequenceEditorToolkitOpened, FLevelSequenceEditorToolkit&);
+	static FLevelSequenceEditorToolkitOpened& OnOpened();
+
+	/** Called when the tab manager is changed */
+	DECLARE_EVENT(FLevelSequenceEditorToolkit, FLevelSequenceEditorToolkitClosed);
+	FLevelSequenceEditorToolkitClosed& OnClosed() { return OnClosedEvent; }
+
 public:
 
 	/**
@@ -66,6 +78,8 @@ public:
 		Collector.AddReferencedObject(LevelSequence);
 	}
 
+	TSharedPtr<ISequencer> GetSequencer() const { return Sequencer; }
+
 private:
 
 	/** Callback for executing the Add Component action. */
@@ -81,7 +95,7 @@ private:
 	TSharedRef<FExtender> HandleMenuExtensibilityGetExtender(const TSharedRef<FUICommandList> CommandList, const TArray<UObject*> ContextSensitiveObjects);
 
 	/** Callback for spawning tabs. */
-	TSharedRef<SDockTab> HandleTabManagerSpawnTab(const FSpawnTabArgs& Args);	
+	TSharedRef<SDockTab> HandleTabManagerSpawnTab(const FSpawnTabArgs& Args);
 
 	/** Callback for the track menu extender. */
 	void HandleTrackMenuExtensionAddTrack(FMenuBuilder& AddTrackMenuBuilder, TArray<UObject*> ContextObjects);
@@ -92,6 +106,11 @@ private:
 	/** Add the specified actors to the sequencer */
 	void AddActorsToSequencer(AActor*const* InActors, int32 NumActors);
 
+private:
+
+	/** Called when this toolkit is to be destroyed */
+	virtual bool OnRequestClose() override;
+
 private:
 
 	/** Level sequence for our edit operation. */
@@ -108,6 +127,9 @@ private:
 	/** Pointer to the style set to use for toolkits. */
 	TSharedRef<ISlateStyle> Style;
 
+	/** Event that is cast when this toolkit is closed */
+	FLevelSequenceEditorToolkitClosed OnClosedEvent;
+
 private:
 
 	/**	The tab ids for all the tabs used */
diff --git a/Engine/Plugins/MovieScene/LevelSequenceEditor/Source/LevelSequenceEditor/Private/Misc/LevelSequenceEditorSpawnRegister.cpp b/Engine/Plugins/MovieScene/LevelSequenceEditor/Source/LevelSequenceEditor/Private/Misc/LevelSequenceEditorSpawnRegister.cpp
--- a/Engine/Plugins/MovieScene/LevelSequenceEditor/Source/LevelSequenceEditor/Private/Misc/LevelSequenceEditorSpawnRegister.cpp
+++ b/Engine/Plugins/MovieScene/LevelSequenceEditor/Source/LevelSequenceEditor/Private/Misc/LevelSequenceEditorSpawnRegister.cpp
@@ -5,8 +5,12 @@
 #include "ISequencerObjectChangeListener.h"
 #include "Editor/LevelEditor/Public/LevelEditor.h"
 #include "LevelSequenceEditorSpawnRegister.h"
+#include "MovieScene.h"
+#include "Kismet2/KismetEditorUtilities.h"
+#include "Kismet2/BlueprintEditorUtils.h"
+#include "AssetSelection.h"
 
-
+#define LOCTEXT_NAMESPACE "LevelSequenceEditorSpawnRegister"
 
 /* FLevelSequenceEditorSpawnRegister structors
  *****************************************************************************/
@@ -169,3 +173,88 @@ void FLevelSequenceEditorSpawnRegister::HandleAnyPropertyChanged(UObject& Spawne
 		DefaultRootComponent->RelativeScale3D = RootComponent->RelativeScale3D;
 	}
 }
+
+
+#if WITH_EDITOR
+
+TValueOrError<FNewSpawnable, FText> FLevelSequenceEditorSpawnRegister::CreateNewSpawnableType(UObject& SourceObject, UMovieScene& OwnerMovieScene)
+{
+	FNewSpawnable NewSpawnable(nullptr, FName::NameToDisplayString(SourceObject.GetName(), false));
+
+	const FName BlueprintName = MakeUniqueObjectName(&OwnerMovieScene, UBlueprint::StaticClass(), SourceObject.GetFName());
+
+	// First off, deal with creating a spawnable from a class
+	if (UClass* InClass = Cast<UClass>(&SourceObject))
+	{
+		if (!InClass->IsChildOf(AActor::StaticClass()))
+		{
+			FText ErrorText = FText::Format(LOCTEXT("NotAnActorClass", "Unable to add spawnable for class of type '{0}' since it is not a valid actor class."), FText::FromString(InClass->GetName()));
+			return MakeError(ErrorText);
+		}
+
+		NewSpawnable.Blueprint = FKismetEditorUtilities::CreateBlueprint( InClass, &OwnerMovieScene, BlueprintName, EBlueprintType::BPTYPE_Normal, UBlueprint::StaticClass(), UBlueprintGeneratedClass::StaticClass() );
+	}
+
+	// Deal with creating a spawnable from an instance of an actor
+	else if (AActor* Actor = Cast<AActor>(&SourceObject))
+	{
+		using namespace EditorUtilities;
+
+		UBlueprint* ActorBlueprint = Cast<UBlueprint>(Actor->GetClass()->ClassGeneratedBy);
+		if (ActorBlueprint)
+		{
+			// Create a new blueprint out of this actor's parent blueprint
+			NewSpawnable.Blueprint = FKismetEditorUtilities::CreateBlueprint(ActorBlueprint->GeneratedClass, &OwnerMovieScene, BlueprintName, EBlueprintType::BPTYPE_Normal, UBlueprint::StaticClass(), UBlueprintGeneratedClass::StaticClass());
+		}
+		else
+		{
+			NewSpawnable.Blueprint = FKismetEditorUtilities::CreateBlueprint(Actor->GetClass(), &OwnerMovieScene, BlueprintName, EBlueprintType::BPTYPE_Normal, UBlueprint::StaticClass(), UBlueprintGeneratedClass::StaticClass(), FName("CreateFromActor"));
+		}
+
+		// Use the actor name
+		NewSpawnable.Name = Actor->GetActorLabel();
+	}
+
+	// If it's a blueprint, we need some special handling
+	else if (UBlueprint* SourceBlueprint = Cast<UBlueprint>(&SourceObject))
+	{
+		UClass* SourceParentClass = SourceBlueprint->GeneratedClass;
+
+		if (!FKismetEditorUtilities::CanCreateBlueprintOfClass(SourceParentClass))
+		{
+			FText ErrorText = FText::Format(LOCTEXT("UnableToAddSpawnableBlueprint", "Unable to add spawnable for class of type '{0}' since it is not a valid blueprint parent class."), FText::FromString(SourceParentClass->GetName()));
+			return MakeError(ErrorText);
+		}
+
+		NewSpawnable.Blueprint = FKismetEditorUtilities::CreateBlueprint(SourceParentClass, &OwnerMovieScene, BlueprintName, EBlueprintType::BPTYPE_Normal, UBlueprint::StaticClass(), UBlueprintGeneratedClass::StaticClass());
+	}
+
+	// At this point we have to assume it's an asset
+	else
+	{
+		// @todo sequencer: Add support for forcing specific factories for an asset?
+		UActorFactory* FactoryToUse = FActorFactoryAssetProxy::GetFactoryForAssetObject(&SourceObject);
+		if (!FactoryToUse)
+		{
+			FText ErrorText = FText::Format(LOCTEXT("CouldNotFindFactory", "Unable to create spawnable from  asset '{0}' - no valid factory could be found."), FText::FromString(SourceObject.GetName()));
+			return MakeError(ErrorText);
+		}
+
+		NewSpawnable.Blueprint = FactoryToUse->CreateBlueprint( &SourceObject, &OwnerMovieScene, BlueprintName );
+	}
+
+
+	if (!NewSpawnable.Blueprint)
+	{
+		FText ErrorText = FText::Format(LOCTEXT("UnknownClassError", "Unable to create a new spawnable object from {0}."), FText::FromString(SourceObject.GetName()));
+		return MakeError(ErrorText);
+	}
+
+	FKismetEditorUtilities::CompileBlueprint(NewSpawnable.Blueprint);
+	return MakeValue(NewSpawnable);
+}
+
+#endif
+
+
+#undef LOCTEXT_NAMESPACE
\ No newline at end of file
diff --git a/Engine/Plugins/MovieScene/LevelSequenceEditor/Source/LevelSequenceEditor/Private/Misc/LevelSequenceEditorSpawnRegister.h b/Engine/Plugins/MovieScene/LevelSequenceEditor/Source/LevelSequenceEditor/Private/Misc/LevelSequenceEditorSpawnRegister.h
--- a/Engine/Plugins/MovieScene/LevelSequenceEditor/Source/LevelSequenceEditor/Private/Misc/LevelSequenceEditorSpawnRegister.h
+++ b/Engine/Plugins/MovieScene/LevelSequenceEditor/Source/LevelSequenceEditor/Private/Misc/LevelSequenceEditorSpawnRegister.h
@@ -40,6 +40,9 @@ public:
 
 	virtual UObject* SpawnObject(const FGuid& BindingId, FMovieSceneSequenceInstance& SequenceInstance, IMovieScenePlayer& Player) override;
 	virtual void PreDestroyObject(UObject& Object, const FGuid& BindingId, FMovieSceneSequenceInstance& SequenceInstance) override;
+#if WITH_EDITOR
+	virtual TValueOrError<FNewSpawnable, FText> CreateNewSpawnableType(UObject& SourceObject, UMovieScene& OwnerMovieScene) override;
+#endif
 
 protected:
 
diff --git a/Engine/Plugins/MovieScene/LevelSequenceEditor/Source/LevelSequenceEditor/Private/Styles/LevelSequenceEditorStyle.h b/Engine/Plugins/MovieScene/LevelSequenceEditor/Source/LevelSequenceEditor/Private/Styles/LevelSequenceEditorStyle.h
--- a/Engine/Plugins/MovieScene/LevelSequenceEditor/Source/LevelSequenceEditor/Private/Styles/LevelSequenceEditorStyle.h
+++ b/Engine/Plugins/MovieScene/LevelSequenceEditor/Source/LevelSequenceEditor/Private/Styles/LevelSequenceEditorStyle.h
@@ -34,6 +34,28 @@ public:
 		Set("LevelSequenceEditor.PossessNewActor.Small", new IMAGE_BRUSH("PossessNewActor_16x", Icon16x16));
 		Set("LevelSequenceEditor.CreateNewLevelSequenceInLevel", new IMAGE_BRUSH("CreateNewLevelSequenceInLevel_16x", Icon16x16));
 		Set("LevelSequenceEditor.CreateNewLevelSequenceInLevel.Small", new IMAGE_BRUSH("CreateNewLevelSequenceInLevel_16x", Icon16x16));
+		Set("LevelSequenceEditor.OnePaneCinematicViewportLayout", new IMAGE_BRUSH("ViewportLayout_OnePaneCinematic", FVector2D(47, 37)));
+		Set("LevelSequenceEditor.TwoPaneCinematicViewportLayout", new IMAGE_BRUSH("ViewportLayout_TwoPaneCinematic", FVector2D(47, 37)));
+		Set("LevelSequenceEditor.OnePaneCinematicViewportLayout.Small", new IMAGE_BRUSH("ViewportLayout_OnePaneCinematic", FVector2D(47, 37)));
+		Set("LevelSequenceEditor.TwoPaneCinematicViewportLayout.Small", new IMAGE_BRUSH("ViewportLayout_TwoPaneCinematic", FVector2D(47, 37)));
+
+		Set("LevelSequenceEditor.CinematicViewportCamera", new IMAGE_BRUSH("CinematicViewportCamera", FVector2D(36, 24)));
+		Set("LevelSequenceEditor.CinematicViewportClose", new IMAGE_BRUSH("CinematicViewportClose", FVector2D(16, 16)));
+
+		Set("LevelSequenceEditor.CinematicViewportPlayMarker", new IMAGE_BRUSH("CinematicViewportPlayMarker", FVector2D(11, 6)));
+		Set("LevelSequenceEditor.CinematicViewportRangeStart", new BORDER_BRUSH("CinematicViewportRangeStart", FMargin(1.f,.3f,0.f,.6f)));
+		Set("LevelSequenceEditor.CinematicViewportRangeEnd", new BORDER_BRUSH("CinematicViewportRangeEnd", FMargin(0.f,.3f,1.f,.6f)));
+
+		Set("LevelSequenceEditor.CinematicViewportSetPlayStart", new IMAGE_BRUSH("CinematicViewportSetPlayStart", FVector2D(24, 24)));
+		Set("LevelSequenceEditor.CinematicViewportSetPlayEnd", new IMAGE_BRUSH("CinematicViewportSetPlayEnd", FVector2D(24, 24)));
+
+		Set("FilmOverlay.DefaultThumbnail", new IMAGE_BRUSH("DefaultFilmOverlayThumbnail", FVector2D(36, 24)));
+
+		Set("FilmOverlay.Disabled", new IMAGE_BRUSH("FilmOverlay.Disabled", FVector2D(36, 24)));
+		Set("FilmOverlay.2x2Grid", new IMAGE_BRUSH("FilmOverlay.2x2Grid", FVector2D(36, 24)));
+		Set("FilmOverlay.3x3Grid", new IMAGE_BRUSH("FilmOverlay.3x3Grid", FVector2D(36, 24)));
+		Set("FilmOverlay.Crosshair", new IMAGE_BRUSH("FilmOverlay.Crosshair", FVector2D(36, 24)));
+		Set("FilmOverlay.Rabatment", new IMAGE_BRUSH("FilmOverlay.Rabatment", FVector2D(36, 24)));
 
 		FSlateStyleRegistry::RegisterSlateStyle(*this);
 	 }
@@ -43,6 +65,18 @@ public:
 	 {
 		FSlateStyleRegistry::UnRegisterSlateStyle(*this);
 	 }
+
+	static TSharedRef<FLevelSequenceEditorStyle> Get()
+	{
+		if (!Singleton.IsValid())
+		{
+			Singleton = MakeShareable(new FLevelSequenceEditorStyle);
+		}
+		return Singleton.ToSharedRef();
+	}
+
+private:
+	static TSharedPtr<FLevelSequenceEditorStyle> Singleton;
 };
 
 
diff --git a/Engine/Source/Developer/AssetTools/Private/AssetTools.cpp b/Engine/Source/Developer/AssetTools/Private/AssetTools.cpp
--- a/Engine/Source/Developer/AssetTools/Private/AssetTools.cpp
+++ b/Engine/Source/Developer/AssetTools/Private/AssetTools.cpp
@@ -391,32 +391,63 @@ UObject* FAssetTools::CreateAsset(UClass* AssetClass, UFactory* Factory, FName C
 		FString AssetName;
 		CreateUniqueAssetName(AssetPath / Factory->GetDefaultNewAssetName(), TEXT(""), PackageName, AssetName);
 
-		FSaveAssetDialogConfig SaveAssetDialogConfig;
-		SaveAssetDialogConfig.DialogTitleOverride = LOCTEXT("SaveAssetDialogTitle", "Save Asset As");
-		SaveAssetDialogConfig.DefaultPath = AssetPath;
-		SaveAssetDialogConfig.DefaultAssetName = AssetName;
-		SaveAssetDialogConfig.ExistingAssetPolicy = ESaveAssetDialogExistingAssetPolicy::AllowButWarn;
-
-		FContentBrowserModule& ContentBrowserModule = FModuleManager::LoadModuleChecked<FContentBrowserModule>("ContentBrowser");
-		FString SaveObjectPath = ContentBrowserModule.Get().CreateModalSaveAssetDialog(SaveAssetDialogConfig);
-		if (!SaveObjectPath.IsEmpty())
+		return CreateAssetWithDialog(AssetName, AssetPath, AssetClass, Factory, CallingContext);
+	}
+
+	return nullptr;
+}
+
+
+UObject* FAssetTools::CreateAssetWithDialog(const FString& AssetName, const FString& PackagePath, UClass* AssetClass, UFactory* Factory, FName CallingContext)
+{
+	FSaveAssetDialogConfig SaveAssetDialogConfig;
+	SaveAssetDialogConfig.DialogTitleOverride = LOCTEXT("SaveAssetDialogTitle", "Save Asset As");
+	SaveAssetDialogConfig.DefaultPath = PackagePath;
+	SaveAssetDialogConfig.DefaultAssetName = AssetName;
+	SaveAssetDialogConfig.ExistingAssetPolicy = ESaveAssetDialogExistingAssetPolicy::AllowButWarn;
+
+	FContentBrowserModule& ContentBrowserModule = FModuleManager::LoadModuleChecked<FContentBrowserModule>("ContentBrowser");
+	FString SaveObjectPath = ContentBrowserModule.Get().CreateModalSaveAssetDialog(SaveAssetDialogConfig);
+	if (!SaveObjectPath.IsEmpty())
+	{
+		FEditorDelegates::OnConfigureNewAssetProperties.Broadcast(Factory);
+		if (Factory->ConfigureProperties())
 		{
-			FEditorDelegates::OnConfigureNewAssetProperties.Broadcast(Factory);
-			if (Factory->ConfigureProperties())
-			{
-				const FString SavePackageName = FPackageName::ObjectPathToPackageName(SaveObjectPath);
-				const FString SaveAssetPath = FPaths::GetPath(SavePackageName);
-				const FString SaveAssetName = FPaths::GetBaseFilename(SavePackageName);
-				FEditorDirectories::Get().SetLastDirectory(ELastDirectory::NEW_ASSET, SaveAssetPath);
+			const FString SavePackageName = FPackageName::ObjectPathToPackageName(SaveObjectPath);
+			const FString SavePackagePath = FPaths::GetPath(SavePackageName);
+			const FString SaveAssetName = FPaths::GetBaseFilename(SavePackageName);
+			FEditorDirectories::Get().SetLastDirectory(ELastDirectory::NEW_ASSET, PackagePath);
 
-				return CreateAsset(SaveAssetName, SaveAssetPath, AssetClass, Factory, CallingContext);
-			}
+			return CreateAsset(SaveAssetName, SavePackagePath, AssetClass, Factory, CallingContext);
 		}
 	}
 
 	return nullptr;
 }
 
+UObject* FAssetTools::DuplicateAssetWithDialog(const FString& AssetName, const FString& PackagePath, UObject* OriginalObject)
+{
+	FSaveAssetDialogConfig SaveAssetDialogConfig;
+	SaveAssetDialogConfig.DialogTitleOverride = LOCTEXT("DuplicateAssetDialogTitle", "Duplicate Asset As");
+	SaveAssetDialogConfig.DefaultPath = PackagePath;
+	SaveAssetDialogConfig.DefaultAssetName = AssetName;
+	SaveAssetDialogConfig.ExistingAssetPolicy = ESaveAssetDialogExistingAssetPolicy::AllowButWarn;
+
+	FContentBrowserModule& ContentBrowserModule = FModuleManager::LoadModuleChecked<FContentBrowserModule>("ContentBrowser");
+	FString SaveObjectPath = ContentBrowserModule.Get().CreateModalSaveAssetDialog(SaveAssetDialogConfig);
+	if (!SaveObjectPath.IsEmpty())
+	{
+		const FString SavePackageName = FPackageName::ObjectPathToPackageName(SaveObjectPath);
+		const FString SavePackagePath = FPaths::GetPath(SavePackageName);
+		const FString SaveAssetName = FPaths::GetBaseFilename(SavePackageName);
+		FEditorDirectories::Get().SetLastDirectory(ELastDirectory::NEW_ASSET, PackagePath);
+
+		return DuplicateAsset(SaveAssetName, SavePackagePath, OriginalObject);
+	}
+
+	return nullptr;
+}
+
 UObject* FAssetTools::DuplicateAsset(const FString& AssetName, const FString& PackagePath, UObject* OriginalObject)
 {
 	// Verify the source object
diff --git a/Engine/Source/Developer/AssetTools/Private/AssetTools.h b/Engine/Source/Developer/AssetTools/Private/AssetTools.h
--- a/Engine/Source/Developer/AssetTools/Private/AssetTools.h
+++ b/Engine/Source/Developer/AssetTools/Private/AssetTools.h
@@ -23,7 +23,9 @@ public:
 	virtual bool GetAssetActions( const TArray<UObject*>& InObjects, FMenuBuilder& MenuBuilder, bool bIncludeHeading = true ) override;
 	virtual UObject* CreateAsset(const FString& AssetName, const FString& PackagePath, UClass* AssetClass, UFactory* Factory, FName CallingContext = NAME_None) override;
 	virtual UObject* CreateAsset(UClass* AssetClass, UFactory* Factory, FName CallingContext = NAME_None) override;
+	virtual UObject* CreateAssetWithDialog(const FString& AssetName, const FString& PackagePath, UClass* AssetClass, UFactory* Factory, FName CallingContext = NAME_None) override;
 	virtual UObject* DuplicateAsset(const FString& AssetName, const FString& PackagePath, UObject* OriginalObject) override;
+	virtual UObject* DuplicateAssetWithDialog(const FString& AssetName, const FString& PackagePath, UObject* OriginalObject) override;
 	virtual void RenameAssets(const TArray<FAssetRenameData>& AssetsAndNames) const override;
 	virtual TArray<UObject*> ImportAssets(const FString& DestinationPath) override;
 	virtual TArray<UObject*> ImportAssets(const TArray<FString>& Files, const FString& DestinationPath, UFactory* ChosenFactory, bool bSyncToBrowser = true) const override;
diff --git a/Engine/Source/Developer/AssetTools/Public/IAssetTools.h b/Engine/Source/Developer/AssetTools/Public/IAssetTools.h
--- a/Engine/Source/Developer/AssetTools/Public/IAssetTools.h
+++ b/Engine/Source/Developer/AssetTools/Public/IAssetTools.h
@@ -106,9 +106,15 @@ public:
 	 */
 	virtual UObject* CreateAsset(const FString& AssetName, const FString& PackagePath, UClass* AssetClass, UFactory* Factory, FName CallingContext = NAME_None) = 0;
 
+	/** Opens an asset picker dialog and creates an asset with the specified name and path */
+	virtual UObject* CreateAssetWithDialog(const FString& AssetName, const FString& PackagePath, UClass* AssetClass, UFactory* Factory, FName CallingContext = NAME_None) = 0;
+
 	/** Opens an asset picker dialog and creates an asset with the chosen path */
 	virtual UObject* CreateAsset(UClass* AssetClass, UFactory* Factory, FName CallingContext = NAME_None) = 0;
 
+	/** Opens an asset picker dialog and creates an asset with the specified name and path. Uses OriginalObject as the duplication source. */
+	virtual UObject* DuplicateAssetWithDialog(const FString& AssetName, const FString& PackagePath, UObject* OriginalObject) = 0;
+
 	/** Creates an asset with the specified name and path. Uses OriginalObject as the duplication source. */
 	virtual UObject* DuplicateAsset(const FString& AssetName, const FString& PackagePath, UObject* OriginalObject) = 0;
 
diff --git a/Engine/Source/Editor/BspMode/BspMode.Build.cs b/Engine/Source/Editor/BspMode/BspMode.Build.cs
--- a/Engine/Source/Editor/BspMode/BspMode.Build.cs
+++ b/Engine/Source/Editor/BspMode/BspMode.Build.cs
@@ -23,6 +23,13 @@ public class BspMode : ModuleRules
         DynamicallyLoadedModuleNames.AddRange(
             new string[] {
 				"PropertyEditor",
+				"PlacementMode",
+			}
+        );
+
+        PrivateIncludePathModuleNames.AddRange(
+            new string[] {
+				"PlacementMode",
 			}
         );
 	}
diff --git a/Engine/Source/Editor/BspMode/Private/BspModeModule.cpp b/Engine/Source/Editor/BspMode/Private/BspModeModule.cpp
--- a/Engine/Source/Editor/BspMode/Private/BspModeModule.cpp
+++ b/Engine/Source/Editor/BspMode/Private/BspModeModule.cpp
@@ -4,6 +4,7 @@
 #include "BspModeModule.h"
 #include "BspMode.h"
 #include "SBspPalette.h"
+#include "IPlacementModeModule.h"
 
 #define LOCTEXT_NAMESPACE "BspMode"
 
@@ -26,6 +27,12 @@ void FBspModeModule::StartupModule()
 	RegisterBspBuilderType(ULinearStairBuilder::StaticClass(), LOCTEXT("LinearStairBuilderName", "Linear Stair"), LOCTEXT("LinearStairBuilderToolTip", "Make a linear stair brush"), FBspModeStyle::Get().GetBrush(TEXT("BspMode.LinearStairBrush")));
 	RegisterBspBuilderType(USpiralStairBuilder::StaticClass(), LOCTEXT("SpiralStairBuilderName", "Spiral Stair"), LOCTEXT("SpiralStairBuilderToolTip", "Make a spiral stair brush"), FBspModeStyle::Get().GetBrush(TEXT("BspMode.SpiralStairBrush")));
 	RegisterBspBuilderType(UTetrahedronBuilder::StaticClass(), LOCTEXT("SphereBuilderName", "Sphere"), LOCTEXT("SphereBuilderToolTip", "Make a sphere brush"), FBspModeStyle::Get().GetBrush(TEXT("BspMode.SphereBrush")));
+
+	IPlacementModeModule& PlacementModeModule = IPlacementModeModule::Get();
+
+	FPlacementCategoryInfo CategoryInfo(LOCTEXT("PlacementMode_Geometry", "Geometry"), "Geometry", TEXT("PMGeometry"), 35);
+	CategoryInfo.CustomGenerator = []() -> TSharedRef<SWidget> { return SNew(SBspPalette); };
+	PlacementModeModule.RegisterPlacementCategory(CategoryInfo);
 }
 
 
@@ -37,12 +44,6 @@ void FBspModeModule::ShutdownModule()
 }
 
 
-TSharedRef< SWidget > FBspModeModule::CreateBspModeWidget() const
-{
-	return SNew(SBspPalette);
-}
-
-
 void FBspModeModule::RegisterBspBuilderType( class UClass* InBuilderClass, const FText& InBuilderName, const FText& InBuilderTooltip, const FSlateBrush* InBuilderIcon )
 {
 	check(InBuilderClass->IsChildOf(UBrushBuilder::StaticClass()));
diff --git a/Engine/Source/Editor/BspMode/Private/BspModeModule.h b/Engine/Source/Editor/BspMode/Private/BspModeModule.h
--- a/Engine/Source/Editor/BspMode/Private/BspModeModule.h
+++ b/Engine/Source/Editor/BspMode/Private/BspModeModule.h
@@ -36,7 +36,6 @@ public:
 	virtual void ShutdownModule() override;
 
 	/** IBspModeModule interface */
-	virtual TSharedRef< SWidget > CreateBspModeWidget() const override;
 	virtual void RegisterBspBuilderType( class UClass* InBuilderClass, const FText& InBuilderName, const FText& InBuilderTooltip, const FSlateBrush* InBuilderIcon ) override;
 	virtual void UnregisterBspBuilderType( class UClass* InBuilderClass ) override;
 
diff --git a/Engine/Source/Editor/BspMode/Public/IBspModeModule.h b/Engine/Source/Editor/BspMode/Public/IBspModeModule.h
--- a/Engine/Source/Editor/BspMode/Public/IBspModeModule.h
+++ b/Engine/Source/Editor/BspMode/Public/IBspModeModule.h
@@ -4,8 +4,6 @@
 class IBspModeModule : public IModuleInterface
 {
 public:
-	virtual TSharedRef< SWidget > CreateBspModeWidget() const = 0;
-
 	virtual void RegisterBspBuilderType( class UClass* InBuilderClass, const FText& InBuilderName, const FText& InBuilderTooltip, const FSlateBrush* InBuilderIcon ) = 0;
 	virtual void UnregisterBspBuilderType( class UClass* InBuilderClass ) = 0;
 };
diff --git a/Engine/Source/Editor/ContentBrowser/Private/ContentBrowserCommands.cpp b/Engine/Source/Editor/ContentBrowser/Private/ContentBrowserCommands.cpp
--- a/Engine/Source/Editor/ContentBrowser/Private/ContentBrowserCommands.cpp
+++ b/Engine/Source/Editor/ContentBrowser/Private/ContentBrowserCommands.cpp
@@ -3,13 +3,17 @@
 #include "ContentBrowserPCH.h"
 #include "ContentBrowserCommands.h"
 
+
 #define LOCTEXT_NAMESPACE "ContentBrowser"
 
+
 void FContentBrowserCommands::RegisterCommands()
 {
-	UI_COMMAND( OpenAssetsOrFolders, "Open Assets or Folders", "Opens the selected assets or folders, depending on the selection", EUserInterfaceActionType::Button, FInputChord(EKeys::Enter) );
-	UI_COMMAND( PreviewAssets, "Preview Assets", "Loads the selected assets and previews them if possible", EUserInterfaceActionType::Button, FInputChord(EKeys::SpaceBar) );
-	UI_COMMAND( DirectoryUp, "Up", "Up to parent directory", EUserInterfaceActionType::Button, FInputChord(EModifierKey::Alt, EKeys::Up) );
+	UI_COMMAND(OpenAssetsOrFolders, "Open Assets or Folders", "Opens the selected assets or folders, depending on the selection", EUserInterfaceActionType::Button, FInputChord(EKeys::Enter));
+	UI_COMMAND(PreviewAssets, "Preview Assets", "Loads the selected assets and previews them if possible", EUserInterfaceActionType::Button, FInputChord(EKeys::SpaceBar));
+	UI_COMMAND(DirectoryUp, "Up", "Up to parent directory", EUserInterfaceActionType::Button, FInputChord(EModifierKey::Alt, EKeys::Up));
+	UI_COMMAND(CreateNewFolder, "Create New Folder", "Creates new folder in selected path", EUserInterfaceActionType::Button, FInputChord(EModifierKey::Control | EModifierKey::Shift, EKeys::N));
 }
 
-#undef LOCTEXT_NAMESPACE
\ No newline at end of file
+
+#undef LOCTEXT_NAMESPACE
diff --git a/Engine/Source/Editor/ContentBrowser/Private/ContentBrowserCommands.h b/Engine/Source/Editor/ContentBrowser/Private/ContentBrowserCommands.h
--- a/Engine/Source/Editor/ContentBrowser/Private/ContentBrowserCommands.h
+++ b/Engine/Source/Editor/ContentBrowser/Private/ContentBrowserCommands.h
@@ -2,19 +2,27 @@
 
 #pragma once
 
-class FContentBrowserCommands : public TCommands<FContentBrowserCommands>
+
+class FContentBrowserCommands
+	: public TCommands<FContentBrowserCommands>
 {
 public:
+
+	/** Default constructor. */
 	FContentBrowserCommands()
 		: TCommands<FContentBrowserCommands>(TEXT("ContentBrowser"), NSLOCTEXT( "ContentBrowser", "ContentBrowser", "Content Browser" ), NAME_None, FEditorStyle::GetStyleSetName() )
-	{
-	}
+	{ }
+
+public:
+
+	//~ TCommands interface
 
 	virtual void RegisterCommands() override;
 
 public:
 
-	TSharedPtr< FUICommandInfo > OpenAssetsOrFolders;
-	TSharedPtr< FUICommandInfo > PreviewAssets;
-	TSharedPtr< FUICommandInfo > DirectoryUp;
-};
\ No newline at end of file
+	TSharedPtr<FUICommandInfo> CreateNewFolder;
+	TSharedPtr<FUICommandInfo> DirectoryUp;
+	TSharedPtr<FUICommandInfo> OpenAssetsOrFolders;
+	TSharedPtr<FUICommandInfo> PreviewAssets;
+};
diff --git a/Engine/Source/Editor/ContentBrowser/Private/ContentBrowserSingleton.cpp b/Engine/Source/Editor/ContentBrowser/Private/ContentBrowserSingleton.cpp
--- a/Engine/Source/Editor/ContentBrowser/Private/ContentBrowserSingleton.cpp
+++ b/Engine/Source/Editor/ContentBrowser/Private/ContentBrowserSingleton.cpp
@@ -291,6 +291,12 @@ void FContentBrowserSingleton::GetSelectedAssets(TArray<FAssetData>& SelectedAss
 	}
 }
 
+void FContentBrowserSingleton::CaptureThumbnailFromViewport(FViewport* InViewport, TArray<FAssetData>& SelectedAssets)
+{
+	ContentBrowserUtils::CaptureThumbnailFromViewport(InViewport, SelectedAssets);
+}
+
+
 void FContentBrowserSingleton::OnEditorLoadSelectedAssetsIfNeeded()
 {
 	if ( PrimaryContentBrowser.IsValid() )
diff --git a/Engine/Source/Editor/ContentBrowser/Private/ContentBrowserSingleton.h b/Engine/Source/Editor/ContentBrowser/Private/ContentBrowserSingleton.h
--- a/Engine/Source/Editor/ContentBrowser/Private/ContentBrowserSingleton.h
+++ b/Engine/Source/Editor/ContentBrowser/Private/ContentBrowserSingleton.h
@@ -33,6 +33,7 @@ public:
 	virtual void SyncBrowserToAssets(const TArray<class FAssetData>& AssetDataList, bool bAllowLockedBrowsers = false, bool bFocusContentBrowser = true) override;
 	virtual void SyncBrowserToAssets(const TArray<UObject*>& AssetList, bool bAllowLockedBrowsers = false, bool bFocusContentBrowser = true) override;
 	virtual void GetSelectedAssets(TArray<FAssetData>& SelectedAssets) override;
+	virtual void CaptureThumbnailFromViewport(FViewport* InViewport, TArray<FAssetData>& SelectedAssets) override;
 
 	/** Gets the content browser singleton as a FContentBrowserSingleton */
 	static FContentBrowserSingleton& Get();
diff --git a/Engine/Source/Editor/ContentBrowser/Private/SContentBrowser.cpp b/Engine/Source/Editor/ContentBrowser/Private/SContentBrowser.cpp
--- a/Engine/Source/Editor/ContentBrowser/Private/SContentBrowser.cpp
+++ b/Engine/Source/Editor/ContentBrowser/Private/SContentBrowser.cpp
@@ -663,26 +663,30 @@ void SContentBrowser::BindCommands()
 	Commands = TSharedPtr< FUICommandList >(new FUICommandList);
 
 	Commands->MapAction(FGenericCommands::Get().Rename, FUIAction(
-		FExecuteAction::CreateSP(this, &SContentBrowser::OnRename),
-		FCanExecuteAction::CreateSP(this, &SContentBrowser::CanRename)
-		));
+		FExecuteAction::CreateSP(this, &SContentBrowser::HandleRenameCommand),
+		FCanExecuteAction::CreateSP(this, &SContentBrowser::HandleRenameCommandCanExecute)
+	));
 
 	Commands->MapAction(FGenericCommands::Get().Delete, FUIAction(
-		FExecuteAction::CreateSP(this, &SContentBrowser::OnDelete),
-		FCanExecuteAction::CreateSP(this, &SContentBrowser::CanDelete)
-		));
+		FExecuteAction::CreateSP(this, &SContentBrowser::HandleDeleteCommandExecute),
+		FCanExecuteAction::CreateSP(this, &SContentBrowser::HandleDeleteCommandCanExecute)
+	));
 
 	Commands->MapAction(FContentBrowserCommands::Get().OpenAssetsOrFolders, FUIAction(
-		FExecuteAction::CreateSP(this, &SContentBrowser::OnOpenAssetsOrFolders)
-		));
+		FExecuteAction::CreateSP(this, &SContentBrowser::HandleOpenAssetsOrFoldersCommandExecute)
+	));
 
 	Commands->MapAction(FContentBrowserCommands::Get().PreviewAssets, FUIAction(
-		FExecuteAction::CreateSP(this, &SContentBrowser::OnPreviewAssets)
-		));
+		FExecuteAction::CreateSP(this, &SContentBrowser::HandlePreviewAssetsCommandEecute)
+	));
 
-	Commands->MapAction( FContentBrowserCommands::Get().DirectoryUp, FUIAction(
-		FExecuteAction::CreateSP( this, &SContentBrowser::OnDirectoryUp )
-		));
+	Commands->MapAction(FContentBrowserCommands::Get().CreateNewFolder, FUIAction(
+		FExecuteAction::CreateSP(this, &SContentBrowser::HandleCreateNewFolderCommandExecute)
+	));
+
+	Commands->MapAction(FContentBrowserCommands::Get().DirectoryUp, FUIAction(
+		FExecuteAction::CreateSP(this, &SContentBrowser::HandleDirectoryUpCommandExecute)
+	));
 }
 
 EVisibility SContentBrowser::GetCollectionViewVisibility() const
@@ -1864,7 +1868,7 @@ FReply SContentBrowser::ForwardClicked()
 	return FReply::Handled();
 }
 
-bool SContentBrowser::CanRename() const
+bool SContentBrowser::HandleRenameCommandCanExecute() const
 {
 	const TArray<TSharedPtr<FAssetViewItem>>& SelectedItems = AssetViewPtr->GetSelectedItems();
 	if (SelectedItems.Num() > 0)
@@ -1882,7 +1886,7 @@ bool SContentBrowser::CanRename() const
 	return false;
 }
 
-void SContentBrowser::OnRename()
+void SContentBrowser::HandleRenameCommand()
 {
 	const TArray<TSharedPtr<FAssetViewItem>>& SelectedItems = AssetViewPtr->GetSelectedItems();
 	if (SelectedItems.Num() > 0)
@@ -1899,7 +1903,7 @@ void SContentBrowser::OnRename()
 	}
 }
 
-bool SContentBrowser::CanDelete() const
+bool SContentBrowser::HandleDeleteCommandCanExecute() const
 {
 	const TArray<TSharedPtr<FAssetViewItem>>& SelectedItems = AssetViewPtr->GetSelectedItems();
 	if (SelectedItems.Num() > 0)
@@ -1917,7 +1921,7 @@ bool SContentBrowser::CanDelete() const
 	return false;
 }
 
-void SContentBrowser::OnDelete()
+void SContentBrowser::HandleDeleteCommandExecute()
 {
 	const TArray<TSharedPtr<FAssetViewItem>>& SelectedItems = AssetViewPtr->GetSelectedItems();
 	if (SelectedItems.Num() > 0)
@@ -1934,23 +1938,34 @@ void SContentBrowser::OnDelete()
 	}
 }
 
-void SContentBrowser::OnOpenAssetsOrFolders()
+void SContentBrowser::HandleOpenAssetsOrFoldersCommandExecute()
 {
 	AssetViewPtr->OnOpenAssetsOrFolders();
 }
 
-void SContentBrowser::OnPreviewAssets()
+void SContentBrowser::HandlePreviewAssetsCommandEecute()
 {
 	AssetViewPtr->OnPreviewAssets();
 }
 
-FReply SContentBrowser::OnDirectoryUpClicked()
+void SContentBrowser::HandleCreateNewFolderCommandExecute()
 {
-	OnDirectoryUp();
-	return FReply::Handled();
+	TArray<FString> SelectedPaths = PathViewPtr->GetSelectedPaths();
+
+	// only create folders when a single path is selected
+	const bool bCanCreateNewFolder = (SelectedPaths.Num() == 1) && ContentBrowserUtils::IsValidPathToCreateNewFolder(SelectedPaths[0]);
+
+	if (bCanCreateNewFolder)
+	{
+		CreateNewFolder(
+			SelectedPaths.Num() > 0
+			? SelectedPaths[0]
+			: FString(),
+			FOnCreateNewFolder::CreateSP(AssetViewPtr.Get(), &SAssetView::OnCreateNewFolder));
+	}
 }
 
-void SContentBrowser::OnDirectoryUp()
+void SContentBrowser::HandleDirectoryUpCommandExecute()
 {
 	TArray<FString> SelectedPaths = PathViewPtr->GetSelectedPaths();
 	if(SelectedPaths.Num() == 1 && !ContentBrowserUtils::IsRootDir(SelectedPaths[0]))
diff --git a/Engine/Source/Editor/ContentBrowser/Private/SContentBrowser.h b/Engine/Source/Editor/ContentBrowser/Private/SContentBrowser.h
--- a/Engine/Source/Editor/ContentBrowser/Private/SContentBrowser.h
+++ b/Engine/Source/Editor/ContentBrowser/Private/SContentBrowser.h
@@ -2,14 +2,18 @@
 
 #pragma once
 
+
 struct FCollectionNameType;
 
+
 /**
  * A widget to display and work with all game and engine content
  */
-class SContentBrowser : public SCompoundWidget
+class SContentBrowser
+	: public SCompoundWidget
 {
 public:
+
 	SLATE_BEGIN_ARGS( SContentBrowser )
 		: _ContainingTab()
 		, _InitiallyLocked(false)
@@ -69,6 +73,7 @@ public:
 	virtual FReply OnMouseButtonDoubleClick( const FGeometry& InMyGeometry, const FPointerEvent& InMouseEvent ) override;
 
 private:
+
 	/** Called to retrieve the text that should be highlighted on assets */
 	FText GetHighlightedText() const;
 
@@ -222,28 +227,28 @@ private:
 	FReply ForwardClicked();
 
 	/** Handler to check to see if a rename command is allowed */
-	bool CanRename() const;
+	bool HandleRenameCommandCanExecute() const;
 
 	/** Handler for Rename */
-	void OnRename();
+	void HandleRenameCommand();
 
 	/** Handler to check to see if a delete command is allowed */
-	bool CanDelete() const;
+	bool HandleDeleteCommandCanExecute() const;
 
 	/** Handler for Delete */
-	void OnDelete();
+	void HandleDeleteCommandExecute();
 
 	/** Handler for opening assets or folders */
-	void OnOpenAssetsOrFolders();
+	void HandleOpenAssetsOrFoldersCommandExecute();
 
 	/** Handler for previewing assets */
-	void OnPreviewAssets();
+	void HandlePreviewAssetsCommandEecute();
 
-	/** Handler for clicking the directory up button */
-	void OnDirectoryUp();
+	/** Handler for creating new folder */
+	void HandleCreateNewFolderCommandExecute();
 
 	/** Handler for clicking the directory up button */
-	FReply OnDirectoryUpClicked();
+	void HandleDirectoryUpCommandExecute();
 
 	/** True if the user may use the history back button */
 	bool IsBackEnabled() const;
@@ -331,6 +336,7 @@ private:
 
 	/** Gets the visibility of the collection view */
 	EVisibility GetCollectionViewVisibility() const;
+
 private:
 
 	/** The tab that contains this browser */
diff --git a/Engine/Source/Editor/ContentBrowser/Public/IContentBrowserSingleton.h b/Engine/Source/Editor/ContentBrowser/Public/IContentBrowserSingleton.h
--- a/Engine/Source/Editor/ContentBrowser/Public/IContentBrowserSingleton.h
+++ b/Engine/Source/Editor/ContentBrowser/Public/IContentBrowserSingleton.h
@@ -8,6 +8,7 @@
 #include "IFilter.h"
 #include "FilterCollection.h"
 #include "AssetThumbnail.h"
+#include "UnrealClient.h"
 
 typedef const FAssetData& FAssetFilterType;
 typedef TFilterCollection<FAssetFilterType> FAssetFilterCollectionType;
@@ -367,4 +368,12 @@ public:
 
 	/** Generates a list of assets that are selected in the primary content browser */
 	virtual void GetSelectedAssets(TArray<FAssetData>& SelectedAssets) = 0;
+
+	/**
+	 * Capture active viewport to thumbnail and assigns that thumbnail to incoming assets
+	 *
+	 * @param InViewport - viewport to sample from
+	 * @param InAssetsToAssign - assets that should receive the new thumbnail ONLY if they are assets that use GenericThumbnails
+	 */
+	virtual void CaptureThumbnailFromViewport(FViewport* InViewport, TArray<FAssetData>& SelectedAssets) = 0;
 };
\ No newline at end of file
