Copying //UE4/Dev-Sequencer to //UE4/Main

==========================
MAJOR FEATURES + CHANGES
==========================

Change 2875025 on 2016/02/20 by [REDACTED_PERSON]

	Sequencer: Cinematic viewport improvements
	  - Added optional letterbox overlay (defaults to 2.35:1)
	  - Added ability to change safe frame colors
	  - Added selected tracks' keys to the transport range
	  - Added buttons for jumping between selected tracks' keyframes on the transport controls
	  - Removed black padding around the viewport where possible
	  - Added ability to specify whether a combo button/menu anchor should close when its parent receives focus
	  - Separated logic of FGroupedKeyArea into FSequencerKeyCollection, so it can be used independently
	  - Added playback range to the viewport frame numbers
	  - All frame numbers are now spin boxes

Change 2875026 on 2016/02/20 by [REDACTED_PERSON]

	Added console commands for recording sequences

	Changed plugin to a developer plugin so we can load it when the editor is in -game mode.
	Added Exec commands.
	Added some more logging to help diagnose problems when using commands.
	Added loading/saving of config for recorder settings (stored in Editor.ini).
	Also disabled controls in recorder window when recording.
	Added auto-saving of assets when in non-editor modes.
	Moved animation settings from UnrealEd to Engine module.

Change 2875036 on 2016/02/20 by [REDACTED_PERSON]

	Sequencer: Call RedrawAllViewports instead of RedrawLevelEditingViewports. In particular, this fixes some update issues when editing values in the key editors.

Change 2875046 on 2016/02/20 by [REDACTED_PERSON]

	Sequencer: Fix so that clicking on UMG Animations doesn't dirty the scene.

Change 2875047 on 2016/02/20 by [REDACTED_PERSON]

	Sequencer: Add option to toggle display of channel colors/lines. View->Channel Colors

Change 2877138 on 2016/02/23 by [REDACTED_PERSON]

	Sequencer: Select corresponding track node when selecting key or section. Removed active/inactive selection since it was only being used in deletion and the rules for deletion are now dependent upon what is selected - delete keys and sections before deleting outliner nodes.

Change 2877143 on 2016/02/23 by [REDACTED_PERSON]

	Added new math function: WindRelativeAnglesDegrees

	Given two angles in degrees, 'wind' the angle in Angle1 so that it avoids >180 degree flips.
	Good for winding rotations previously expressed as quaternions into a euler-angle representation.

Change 2877147 on 2016/02/23 by [REDACTED_PERSON]

	Added the ability to import sequencer transforms from the root node of an animation sequence

	Intended for use after re-importing animations from DCC tools.
	Available in the right-click menu for transform tracks.
	Also added FindTrackBinding to UMovieScene so track bindings can be recovered from tracks.

Change 2877163 on 2016/02/23 by [REDACTED_PERSON]

	Sequencer: Add option to create keyframe sections as infinite. Sequencer defaults to true, UMG defaults to false.

Change 2877165 on 2016/02/23 by [REDACTED_PERSON]

	Sequencer: Drawing vertical position lines when dragging keys

Change 2878748 on 2016/02/23 by [REDACTED_PERSON]

	Curve Editor: Switch curve type to user when flatting or straightening tangents.

Change 2878799 on 2016/02/23 by [REDACTED_PERSON]

	Sequencer - Add folders support to the outliner.

Change 2880769 on 2016/02/24 by [REDACTED_PERSON]

	Sequencer: Added ability to override runtime spawnable ownership in sequencer
	  - This is exposed as an option on spawnables "Keep Alive Outside Playback Range (In Sequencer)"
	  - Enabling this will stop spawnables from being destroyed when scrubbing outside of the playback range

Change 2880770 on 2016/02/24 by [REDACTED_PERSON]

	Sequencer: Added countdown and recording indicator display when recording

	Also fixed extra popups added post-PIE when animation recordings auto shutdown.

Change 2880782 on 2016/02/24 by [REDACTED_PERSON]

	Sequencer: Snapping now also uses the current time as a possible snap time.

Change 2880793 on 2016/02/24 by [REDACTED_PERSON]

	Sequencer: Add +Animation to Animation track so that it's consistent with all other tracks that have a + button.

Change 2880812 on 2016/02/24 by [REDACTED_PERSON]

	Sequencer: Fix adjusting the leading edge of a shot section so that it cuts into the shot rather than adjusts the start time.

Change 2881624 on 2016/02/25 by [REDACTED_PERSON]

	Changing shader version GUID to fix corrupt shaders in ddc

Change 2882408 on 2016/02/25 by [REDACTED_PERSON]

	Asset/actors stored in TLazyObjectPtrs can now reference game content from engine

	This is a legitimate use case as lazy object ptrs are designed to reference assets/actors cross-domain.

Change 2882409 on 2016/02/25 by [REDACTED_PERSON]

diff --git a/Engine/Config/BaseEditorPerProjectUserSettings.ini b/Engine/Config/BaseEditorPerProjectUserSettings.ini
--- a/Engine/Config/BaseEditorPerProjectUserSettings.ini
+++ b/Engine/Config/BaseEditorPerProjectUserSettings.ini
@@ -727,8 +727,11 @@ bShowRangeSlider=true
 bKeepPlayRangeInSectionBounds=false
 TimeSnapInterval=0.033334
 ZeroPadFrames=4
+bInfiniteKeyAreas=true
 
 [/Script/LevelSequenceEditor.LevelSequenceEditorSettings]
++TrackSettings=(MatchingActorClass=/Script/Engine.StaticMeshActor,DefaultTracks=(/Script/MovieSceneTracks.MovieScene3DTransformTrack))
++TrackSettings=(MatchingActorClass=/Script/Engine.SkeletalMeshActor,DefaultTracks=(/Script/MovieSceneTracks.MovieScene3DTransformTrack,/Script/MovieSceneTracks.MovieSceneSkeletalAnimationTrack))
 +TrackSettings=(MatchingActorClass=/Script/Engine.CameraActor,DefaultTracks=(/Script/MovieSceneTracks.MovieScene3DTransformTrack))
-;+TrackSettings=(MatchingActorClass=/Script/Engine.CineCameraActor,DefaultPropertyTracks=((ComponentPath="CameraComponent",PropertyPath="CurrentFocalLength")))
++TrackSettings=(MatchingActorClass=/Script/CinematicCamera.CineCameraActor,DefaultPropertyTracks=((ComponentPath="CameraComponent",PropertyPath="CurrentFocalLength"),(ComponentPath="CameraComponent",PropertyPath="FocusSettings.ManualFocusDistance"),(ComponentPath="CameraComponent",PropertyPath="CurrentAperture")))
 +TrackSettings=(MatchingActorClass=/Script/Engine.Light,DefaultTracks=(None),DefaultPropertyTracks=((ComponentPath="LightComponent0",PropertyPath="Intensity"),(ComponentPath="LightComponent0",PropertyPath="LightColor")))
diff --git a/Engine/Config/BaseEngine.ini b/Engine/Config/BaseEngine.ini
--- a/Engine/Config/BaseEngine.ini
+++ b/Engine/Config/BaseEngine.ini
@@ -2563,24 +2563,28 @@ bBuildForD3D12=false
 +Settings=(ScreenSize=.3,OptimizationSettings=(NumOfTrianglesPercentage=.25))
 +Settings=(ScreenSize=.1,OptimizationSettings=(NumOfTrianglesPercentage=.125))
 +Settings=(ScreenSize=.06,OptimizationSettings=(NumOfTrianglesPercentage=.06))
+
 [/Script/CinematicCamera.CineCameraComponent]
 +FilmbackPresets=(Name="Super 8mm",FilmbackSettings=(SensorWidth=5.79,SensorHeight=4.01))
-+FilmbackPresets=(Name="Super 16mm",FilmbackSettings=(SensorWidth=12.52,SensorHeight=7.41))
-+FilmbackPresets=(Name="Super 35mm 3 Perf",FilmbackSettings=(SensorWidth=24.9,SensorHeight=12))
-+FilmbackPresets=(Name="Super 35mm 4 Perf",FilmbackSettings=(SensorWidth=24.89,SensorHeight=18.67))
-+FilmbackPresets=(Name="Four Thirds",FilmbackSettings=(SensorWidth=17.3,SensorHeight=13))
-+FilmbackPresets=(Name="35mm Full Frame",FilmbackSettings=(SensorWidth=36,SensorHeight=24))
-+FilmbackPresets=(Name="APS-C",FilmbackSettings=(SensorWidth=23.6,SensorHeight=15.7))
-+FilmbackPresets=(Name="APS-C (Canon)",FilmbackSettings=(SensorWidth=22.2,SensorHeight=14.8))
-+FilmbackPresets=(Name="APS-H",FilmbackSettings=(SensorWidth=27.9,SensorHeight=18.6))
-+FilmbackPresets=(Name="65mm",FilmbackSettings=(SensorWidth=52.48,SensorHeight=23.01))
++FilmbackPresets=(Name="Super 16mm",FilmbackSettings=(SensorWidth=12.52,SensorHeight=7.58))
++FilmbackPresets=(Name="Super 35mm",FilmbackSettings=(SensorWidth=24.89,SensorHeight=18.66))
++FilmbackPresets=(Name="35mm Academy",FilmbackSettings=(SensorWidth=21.946,SensorHeight=16.002))
++FilmbackPresets=(Name="35mm Full Aperture",FilmbackSettings=(SensorWidth=24.892,SensorHeight=18.9121))
++FilmbackPresets=(Name="35mm VistaVision",FilmbackSettings=(SensorWidth=25.146,SensorHeight=37.719))
+;+FilmbackPresets=(Name="65mm 15 perf",FilmbackSettings=(SensorWidth=52.48,SensorHeight=23.01))
 +FilmbackPresets=(Name="IMAX 70mm",FilmbackSettings=(SensorWidth=70.41,SensorHeight=56.63))
-+LensPresets=(Name="12mm Prime f/2.8",LensSettings=(MinFocalLength=12,MaxFocalLength=12,MinFStop=2.8,MaxFStop=22)
-+LensPresets=(Name="30mm Prime f/1.4",LensSettings=(MinFocalLength=30,MaxFocalLength=30,MinFStop=1.4,MaxFStop=22)
-+LensPresets=(Name="50mm Prime f/1.8",LensSettings=(MinFocalLength=50,MaxFocalLength=50,MinFStop=1.8,MaxFStop=22)
-+LensPresets=(Name="85mm Prime f/1.8",LensSettings=(MinFocalLength=85,MaxFocalLength=85,MinFStop=1.8,MaxFStop=22)
-+LensPresets=(Name="105mm Prime f/2",LensSettings=(MinFocalLength=105,MaxFocalLength=105,MinFStop=2,MaxFStop=22,)
-+LensPresets=(Name="200mm Prime f/2",LensSettings=(MinFocalLength=200,MaxFocalLength=200,MinFStop=2,MaxFStop=22)
++FilmbackPresets=(Name="APS-C (Canon)",FilmbackSettings=(SensorWidth=22.2,SensorHeight=14.8))
++FilmbackPresets=(Name="Full Frame DSLR",FilmbackSettings=(SensorWidth=36,SensorHeight=24))
++FilmbackPresets=(Name="Micro Four Thirds",FilmbackSettings=(SensorWidth=17.3,SensorHeight=13))
+; Sony F65
+; Arri Alexa
+
++LensPresets=(Name="12mm Prime f/2.8",LensSettings=(MinFocalLength=12,MaxFocalLength=12,MinFStop=2.8,MaxFStop=22))
++LensPresets=(Name="30mm Prime f/1.4",LensSettings=(MinFocalLength=30,MaxFocalLength=30,MinFStop=1.4,MaxFStop=22))
++LensPresets=(Name="50mm Prime f/1.8",LensSettings=(MinFocalLength=50,MaxFocalLength=50,MinFStop=1.8,MaxFStop=22))
++LensPresets=(Name="85mm Prime f/1.8",LensSettings=(MinFocalLength=85,MaxFocalLength=85,MinFStop=1.8,MaxFStop=22))
++LensPresets=(Name="105mm Prime f/2",LensSettings=(MinFocalLength=105,MaxFocalLength=105,MinFStop=2,MaxFStop=22))
++LensPresets=(Name="200mm Prime f/2",LensSettings=(MinFocalLength=200,MaxFocalLength=200,MinFStop=2,MaxFStop=22))
 +LensPresets=(Name="24-70mm Zoom f/2.8",LensSettings=(MinFocalLength=24,MaxFocalLength=70,MinFStop=2.8,MaxFStop=22))
 +LensPresets=(Name="70-200mm Zoom f/2.8",LensSettings=(MinFocalLength=70,MaxFocalLength=200,MinFStop=2.8,MaxFStop=22))
 
diff --git a/Engine/Config/BaseInput.ini b/Engine/Config/BaseInput.ini
--- a/Engine/Config/BaseInput.ini
+++ b/Engine/Config/BaseInput.ini
@@ -174,4 +174,5 @@ MaxScrollbackSize=1024
 +ManualAutoCompleteList=(Command="Log reset",Desc="Undo any changes to log verbosity")
 +ManualAutoCompleteList=(Command="RecordAnimation ActorName AnimName",Desc="Record animation for a specified actor to the specified file")
 +ManualAutoCompleteList=(Command="StopRecordingAnimation All",Desc="Stop all recording animations")
-+ManualAutoCompleteList=(Command="StopRecordingAnimation ActorName",Desc="Stop recording the specifed animation")
++ManualAutoCompleteList=(Command="RecordSequence Filter ActorOrClassName",Desc="Record a level sequence from gameplay. Filter=<All|Actor|Class>")
++ManualAutoCompleteList=(Command="StopRecordingSequence",Desc="Stop recording the current sequence. Only one sequence recording can be active at one time.")
diff --git a/Engine/Plugins/Messaging/UdpMessaging/Source/UdpMessaging/Private/Transport/UdpDeserializedMessage.cpp b/Engine/Plugins/Messaging/UdpMessaging/Source/UdpMessaging/Private/Transport/UdpDeserializedMessage.cpp
--- a/Engine/Plugins/Messaging/UdpMessaging/Source/UdpMessaging/Private/Transport/UdpDeserializedMessage.cpp
+++ b/Engine/Plugins/Messaging/UdpMessaging/Source/UdpMessaging/Private/Transport/UdpDeserializedMessage.cpp
@@ -5,6 +5,23 @@
 #include "StructDeserializer.h"
 
 
+/* FUdpDeserializedMessage structors
+ *****************************************************************************/
+
+FUdpDeserializedMessage::~FUdpDeserializedMessage()
+{
+	if (MessageData != nullptr)
+	{
+		if (TypeInfo.IsValid())
+		{
+			TypeInfo->DestroyStruct(MessageData);
+		}
+
+		FMemory::Free(MessageData);
+	}
+}
+
+
 /* FUdpDeserializedMessage interface
  *****************************************************************************/
 
diff --git a/Engine/Plugins/Messaging/UdpMessaging/Source/UdpMessaging/Private/Transport/UdpDeserializedMessage.h b/Engine/Plugins/Messaging/UdpMessaging/Source/UdpMessaging/Private/Transport/UdpDeserializedMessage.h
--- a/Engine/Plugins/Messaging/UdpMessaging/Source/UdpMessaging/Private/Transport/UdpDeserializedMessage.h
+++ b/Engine/Plugins/Messaging/UdpMessaging/Source/UdpMessaging/Private/Transport/UdpDeserializedMessage.h
@@ -18,8 +18,12 @@ public:
 	 */
 	FUdpDeserializedMessage(const IMessageAttachmentPtr& InAttachment)
 		: Attachment(InAttachment)
+		, MessageData(nullptr)
 	{ }
 
+	/** Virtual destructor. */
+	virtual ~FUdpDeserializedMessage() override;
+
 public:
 
 	/**
diff --git a/Engine/Plugins/MovieScene/LevelSequenceEditor/Source/LevelSequenceEditor/Private/CinematicViewport/CinematicViewportLayoutEntity.h b/Engine/Plugins/MovieScene/LevelSequenceEditor/Source/LevelSequenceEditor/Private/CinematicViewport/CinematicViewportLayoutEntity.h
new file mode 100644
--- /dev/null
+++ b/Engine/Plugins/MovieScene/LevelSequenceEditor/Source/LevelSequenceEditor/Private/CinematicViewport/CinematicViewportLayoutEntity.h
@@ -0,0 +1,40 @@
+// Copyright 1998-2015 Epic Games, Inc. All Rights Reserved.
+
+#pragma once
+
+#include "LevelViewportLayout.h"
+#include "SCinematicLevelViewport.h"
+
+class FCinematicViewportLayoutEntity : public IViewportLayoutEntity
+{
+public:
+	FCinematicViewportLayoutEntity(const FViewportConstructionArgs& Args)
+		: Widget(
+			SNew(SCinematicLevelViewport)
+			.ParentLayout(Args.ParentLayout)
+			.ParentLevelEditor(Args.ParentLevelEditor)
+			.LayoutName(*Args.ConfigKey)
+			)
+	{}
+
+	virtual TSharedRef<SWidget> AsWidget() const override { return Widget; }
+	virtual TSharedPtr<SLevelViewport> AsLevelViewport() const { return Widget->GetLevelViewport(); }
+
+	FName GetType() const
+	{
+		static FName CinematicName("Cinematic");
+		return CinematicName;
+	}
+
+	FLevelEditorViewportClient& GetLevelViewportClient() const { return Widget->GetLevelViewport()->GetLevelViewportClient(); }
+	bool IsPlayInEditorViewportActive() const { return false; }
+	void RegisterGameViewportIfPIE(){}
+	void SetKeyboardFocus(){ FSlateApplication::Get().SetKeyboardFocus(Widget); }
+	void OnLayoutDestroyed(){}
+	void SaveConfig(const FString&){}
+
+protected:
+
+	/** This entity's widget */
+	TSharedRef<SCinematicLevelViewport> Widget;
+};
\ No newline at end of file
diff --git a/Engine/Plugins/MovieScene/LevelSequenceEditor/Source/LevelSequenceEditor/Private/CinematicViewport/FilmOverlays.cpp b/Engine/Plugins/MovieScene/LevelSequenceEditor/Source/LevelSequenceEditor/Private/CinematicViewport/FilmOverlays.cpp
--- a/Engine/Plugins/MovieScene/LevelSequenceEditor/Source/LevelSequenceEditor/Private/CinematicViewport/FilmOverlays.cpp
+++ b/Engine/Plugins/MovieScene/LevelSequenceEditor/Source/LevelSequenceEditor/Private/CinematicViewport/FilmOverlays.cpp
@@ -5,9 +5,67 @@
 #include "LevelSequenceEditorStyle.h"
 #include "Widgets/Colors/SColorPicker.h"
 #include "NumericTypeInterface.h"
+#include "SColorBlock.h"
 
 #define LOCTEXT_NAMESPACE "LevelSequenceEditorFilmOverlays"
 
+DECLARE_DELEGATE_OneParam(FOnColorPicked, const FLinearColor&);
+
+namespace WidgetHelpers
+{
+	TSharedRef<SWidget> CreateColorWidget(FLinearColor* Color, FOnColorPicked OnColorPicked = FOnColorPicked())
+	{
+		auto GetValue = [=]{
+			return *Color;
+		};
+
+		auto SetValue = [=](FLinearColor NewColor){
+			*Color = NewColor;
+			OnColorPicked.ExecuteIfBound(NewColor);
+		};
+
+		auto OnGetMenuContent = [=]() -> TSharedRef<SWidget> {
+			// Open a color picker
+			return SNew(SColorPicker)
+			.TargetColorAttribute_Lambda(GetValue)
+			.UseAlpha(true)
+			.DisplayInlineVersion(true)
+			.OnColorCommitted_Lambda(SetValue);
+		};
+
+		return SNew(SComboButton)
+			.ContentPadding(0)
+			.HasDownArrow(false)
+			.ButtonStyle(FEditorStyle::Get(), "Sequencer.AnimationOutliner.ColorStrip")
+			.OnGetMenuContent_Lambda(OnGetMenuContent)
+			.CollapseMenuOnParentFocus(true)
+			.ButtonContent()
+			[
+				SNew(SColorBlock)
+				.Color_Lambda(GetValue)
+				.ShowBackgroundForAlpha(true)
+				.IgnoreAlpha(false)
+				.Size(FVector2D(10.0f, 10.0f))
+			];
+	}
+
+	template<typename T>
+	TSharedRef<SWidget> CreateSpinBox(T* Value, T Min, T Max, TSharedPtr<INumericTypeInterface<T>> TypeInterface = nullptr)
+	{
+		auto GetValue = [=]{ return *Value; };
+		auto SetValue = [=](T NewValue){ *Value = NewValue; };
+		auto SetValueCommitted = [=](T NewValue, ETextCommit::Type){ *Value = NewValue; };
+
+		return SNew(SSpinBox<T>)
+			.MinValue(Min)
+			.MaxValue(Max)
+			.Value_Lambda(GetValue)
+			.OnValueChanged_Lambda(SetValue)
+			.OnValueCommitted_Lambda(SetValueCommitted)
+			.TypeInterface(TypeInterface);
+	}
+}
+
 struct FFilmOverlay_None : IFilmOverlay
 {
 	FText GetDisplayName() const { return LOCTEXT("OverlayDisabled", "Disabled"); }
@@ -198,7 +256,7 @@ struct FFilmOverlay_Crosshair : IFilmOverlay
 
 struct FFilmOverlay_SafeFrame : IFilmOverlay
 {
-	FFilmOverlay_SafeFrame(FText InDisplayName, float InSizePercentage, const FLinearColor& InColor)
+	FFilmOverlay_SafeFrame(FText InDisplayName, int32 InSizePercentage, const FLinearColor& InColor)
 	{
 		DisplayName = InDisplayName;
 		SizePercentage = InSizePercentage;
@@ -236,13 +294,13 @@ struct FFilmOverlay_SafeFrame : IFilmOverlay
 
 	virtual TSharedPtr<SWidget> ConstructSettingsWidget() override
 	{
-		struct FPercentageInterface : TDefaultNumericTypeInterface<float>
+		struct FPercentageInterface : TDefaultNumericTypeInterface<int32>
 		{
-			virtual FString ToString(const float& Value) const override
+			virtual FString ToString(const int32& Value) const override
 			{
 				return TDefaultNumericTypeInterface::ToString(Value) + TEXT("%");
 			}
-			virtual TOptional<float> FromString(const FString& InString, const float& InExistingValue) override
+			virtual TOptional<int32> FromString(const FString& InString, const int32& InExistingValue) override
 			{
 				return TDefaultNumericTypeInterface::FromString(InString.Replace(TEXT("%"), TEXT("")), InExistingValue);
 			}
@@ -252,24 +310,142 @@ struct FFilmOverlay_SafeFrame : IFilmOverlay
 			}
 		};
 
-		auto GetValue = [this]{ return SizePercentage; };
-		auto SetValue = [this](float NewValue){ SizePercentage = NewValue; };
-		auto SetValueCommitted = [this](float NewValue, ETextCommit::Type){ SizePercentage = NewValue; };
+		return SNew(SVerticalBox)
 
-		return SNew(SSpinBox<float>)
-			.MinValue(1)
-			.MaxValue(99)
-			.Value_Lambda(GetValue)
-			.OnValueChanged_Lambda(SetValue)
-			.OnValueCommitted_Lambda(SetValueCommitted)
-			.TypeInterface(MakeShareable(new FPercentageInterface));
+			+ SVerticalBox::Slot()
+			.AutoHeight()
+			[
+				WidgetHelpers::CreateSpinBox<int32>(&SizePercentage, 1, 99, MakeShareable(new FPercentageInterface))
+			]
+
+			+ SVerticalBox::Slot()
+			.Padding(FMargin(0.f, 1.f, 0.f, 0.f))
+			.AutoHeight()
+			[
+				WidgetHelpers::CreateColorWidget(&Tint)
+			];
 	}
 
 private:
-	float SizePercentage;
+	int32 SizePercentage;
 	FText DisplayName;
 };
 
+
+struct FFilmOverlay_LetterBox : IFilmOverlay
+{
+	FFilmOverlay_LetterBox()
+	{
+		Ratio1 = 2.35f;
+		Ratio2 = 1.f;
+		Color = FLinearColor(0.f, 0.f, 0.f, .5f);
+	}
+
+	FText GetDisplayName() const { return LOCTEXT("LetterboxMask", "Letterbox Mask"); }
+
+	const FSlateBrush* GetThumbnail() const { return nullptr; }
+
+	void Paint(const FGeometry& AllottedGeometry, const FSlateRect& MyClippingRect, FSlateWindowElementList& OutDrawElements, int32 LayerId) const
+	{
+		const FSlateBrush* Brush = FEditorStyle::GetBrush("WhiteBrush");
+
+		const float DesiredRatio = Ratio1 / Ratio2;
+		const float CurrentRatio = AllottedGeometry.Size.X / AllottedGeometry.Size.Y;
+
+		if (CurrentRatio > DesiredRatio)
+		{
+			// vertical letterbox mask
+			FVector2D LetterBoxSize((AllottedGeometry.Size.X - AllottedGeometry.Size.Y * DesiredRatio) * .5f, AllottedGeometry.Size.Y);
+			FSlateDrawElement::MakeBox(
+				OutDrawElements,
+				LayerId,
+				AllottedGeometry.ToPaintGeometry(LetterBoxSize, FSlateLayoutTransform(FVector2D(0.f, 0.f))),
+				Brush,
+				MyClippingRect,
+				ESlateDrawEffect::None,
+				Color
+				);
+
+			FSlateDrawElement::MakeBox(
+				OutDrawElements,
+				LayerId,
+				AllottedGeometry.ToPaintGeometry(LetterBoxSize, FSlateLayoutTransform(FVector2D(AllottedGeometry.Size.X - LetterBoxSize.X, 0.f))),
+				Brush,
+				MyClippingRect,
+				ESlateDrawEffect::None,
+				Color
+				);
+
+		}
+		else if (CurrentRatio < DesiredRatio)
+		{
+			// horizontal letterbox mask
+			FVector2D LetterBoxSize(AllottedGeometry.Size.X, (AllottedGeometry.Size.Y - AllottedGeometry.Size.X / DesiredRatio) * .5f);
+
+			FSlateDrawElement::MakeBox(
+				OutDrawElements,
+				LayerId,
+				AllottedGeometry.ToPaintGeometry(LetterBoxSize, FSlateLayoutTransform(FVector2D(0.f, 0.f))),
+				Brush,
+				MyClippingRect,
+				ESlateDrawEffect::None,
+				Color
+				);
+
+			FSlateDrawElement::MakeBox(
+				OutDrawElements,
+				LayerId,
+				AllottedGeometry.ToPaintGeometry(LetterBoxSize, FSlateLayoutTransform(FVector2D(0.f, AllottedGeometry.Size.Y - LetterBoxSize.Y))),
+				Brush,
+				MyClippingRect,
+				ESlateDrawEffect::None,
+				Color
+				);
+		}
+	}
+
+	virtual TSharedPtr<SWidget> ConstructSettingsWidget() override
+	{
+		return SNew(SVerticalBox)
+
+			+ SVerticalBox::Slot()
+			.AutoHeight()
+			[
+				SNew(SHorizontalBox)
+
+				+ SHorizontalBox::Slot()
+				[
+					WidgetHelpers::CreateSpinBox<float>(&Ratio1, 0.1, 35)
+				]
+
+				+ SHorizontalBox::Slot()
+				.AutoWidth()
+				.Padding(FMargin(2.f, 0.f))
+				[
+					SNew(STextBlock)
+					.Text(FText::FromString(TEXT(":")))
+				]
+
+				+ SHorizontalBox::Slot()
+				[
+					WidgetHelpers::CreateSpinBox<float>(&Ratio2, 0.1, 35)
+				]
+			]
+
+			+ SVerticalBox::Slot()
+			.Padding(FMargin(0.f, 1.f, 0.f, 0.f))
+			.AutoHeight()
+			[
+				WidgetHelpers::CreateColorWidget(&Color)
+			];
+	}
+
+private:
+	float Ratio1;
+	float Ratio2;
+	FLinearColor Color;
+};
+
 void SFilmOverlay::Construct(const FArguments& InArgs)
 {
 	FilmOverlays = InArgs._FilmOverlays;
@@ -300,6 +476,7 @@ void SFilmOverlayOptions::Construct(const FArguments& InArgs)
 	ToggleableOverlays.Add("ActionSafeFrame",	TUniquePtr<IFilmOverlay>(new FFilmOverlay_SafeFrame(LOCTEXT("ActionSafeFrame", "Action Safe"), 95.f, FLinearColor::Red)));
 	ToggleableOverlays.Add("TitleSafeFrame",	TUniquePtr<IFilmOverlay>(new FFilmOverlay_SafeFrame(LOCTEXT("TitleSafeFrame", "Title Safe"), 90.f, FLinearColor::Yellow)));
 	ToggleableOverlays.Add("CustomSafeFrame",	TUniquePtr<IFilmOverlay>(new FFilmOverlay_SafeFrame(LOCTEXT("CustomSafeFrame", "Custom Safe"), 85.f, FLinearColor::Green)));
+	ToggleableOverlays.Add("LetterBox",			TUniquePtr<IFilmOverlay>(new FFilmOverlay_LetterBox));
 
 	OverlayWidget = SNew(SFilmOverlay)
 		.Visibility(EVisibility::HitTestInvisible)
@@ -330,10 +507,8 @@ FLinearColor SFilmOverlayOptions::GetMasterColorTint() const
 	return MasterColorTint;
 }
 
-void SFilmOverlayOptions::SetMasterColorTint(FLinearColor Tint)
+void SFilmOverlayOptions::OnMasterColorTintChanged(const FLinearColor& Tint)
 {
-	MasterColorTint = Tint;
-
 	IFilmOverlay* Overlay = GetMasterFilmOverlay();
 	if (Overlay)
 	{
@@ -368,40 +543,27 @@ TSharedRef<SWidget> SFilmOverlayOptions::GetMenuContent()
 		]
 
 		+ SGridPanel::Slot(0, 1)
+		.Padding(10)
 		[
-			SNew(SVerticalBox)
+			SNew(SHorizontalBox)
 
-			+ SVerticalBox::Slot()
-			.AutoHeight()
+			+ SHorizontalBox::Slot()
+			.Padding(2.f, 0.f, 5.f, 0.f)
+			.AutoWidth()
+			.VAlign(VAlign_Center)
 			[
-				SNew(SHeader)
-				[
-					SNew(STextBlock)
-					.Text(LOCTEXT("ColorPickerHeader", "Overlay Color Tint"))
-				]
+				SNew(STextBlock)
+				.Text(LOCTEXT("OverlayTint", "Tint: "))
 			]
 
-			+ SVerticalBox::Slot()
-			.AutoHeight()
-			.Padding(10)
-			.HAlign(HAlign_Center)
+			+ SHorizontalBox::Slot()
+			.VAlign(VAlign_Center)
 			[
-				SNew(SColorPicker)
-				.TargetColorAttribute(this, &SFilmOverlayOptions::GetMasterColorTint)
-				.UseAlpha(true)
-				.DisplayInlineVersion(true)
-				.OnColorCommitted(this, &SFilmOverlayOptions::SetMasterColorTint)
+				WidgetHelpers::CreateColorWidget(&MasterColorTint, FOnColorPicked::CreateRaw(this, &SFilmOverlayOptions::OnMasterColorTintChanged))
 			]
 		]
 
-		// @todo: specific overlay settings to go here
-		// + SGridPanel::Slot(1, 0)
-		// [
-		// 
-		// ]
-
 		+ SGridPanel::Slot(1, 0)
-		//+ SGridPanel::Slot(1, 1)
 		[
 			SNew(SVerticalBox)
 
@@ -412,7 +574,7 @@ TSharedRef<SWidget> SFilmOverlayOptions::GetMenuContent()
 				SNew(SHeader)
 				[
 					SNew(STextBlock)
-					.Text(LOCTEXT("SafeFrameHeader", "Safe Frames"))
+					.Text(LOCTEXT("SafeFrameHeader", "Frames"))
 				]
 			]
 			
diff --git a/Engine/Plugins/MovieScene/LevelSequenceEditor/Source/LevelSequenceEditor/Private/CinematicViewport/FilmOverlays.h b/Engine/Plugins/MovieScene/LevelSequenceEditor/Source/LevelSequenceEditor/Private/CinematicViewport/FilmOverlays.h
--- a/Engine/Plugins/MovieScene/LevelSequenceEditor/Source/LevelSequenceEditor/Private/CinematicViewport/FilmOverlays.h
+++ b/Engine/Plugins/MovieScene/LevelSequenceEditor/Source/LevelSequenceEditor/Private/CinematicViewport/FilmOverlays.h
@@ -102,7 +102,7 @@ private:
 
 	/** Get/Set the color tint override for the current master overlay */
 	FLinearColor GetMasterColorTint() const;
-	void SetMasterColorTint(FLinearColor Tint);
+	void OnMasterColorTintChanged(const FLinearColor& Tint);
 
 private:
 
diff --git a/Engine/Plugins/MovieScene/LevelSequenceEditor/Source/LevelSequenceEditor/Private/CinematicViewport/SCinematicLevelViewport.cpp b/Engine/Plugins/MovieScene/LevelSequenceEditor/Source/LevelSequenceEditor/Private/CinematicViewport/SCinematicLevelViewport.cpp
--- a/Engine/Plugins/MovieScene/LevelSequenceEditor/Source/LevelSequenceEditor/Private/CinematicViewport/SCinematicLevelViewport.cpp
+++ b/Engine/Plugins/MovieScene/LevelSequenceEditor/Source/LevelSequenceEditor/Private/CinematicViewport/SCinematicLevelViewport.cpp
@@ -19,6 +19,7 @@
 #include "CineCameraComponent.h"
 #include "UnitConversion.h"
 #include "ScopedTransaction.h"
+#include "SLevelViewport.h"
 
 #define LOCTEXT_NAMESPACE "SCinematicLevelViewport"
 
@@ -80,21 +81,6 @@ FCinematicViewportClient::FCinematicViewportClient()
 	bDisableInput = false;
 }
 
-bool FCinematicViewportClient::InputKey(FViewport* Viewport, int32 ControllerId, FKey Key, EInputEvent Event, float AmountDepressed, bool bGamepad)
-{
-	// Allow immersive mode to be processed when input handling is disabled
-	if (bDisableInput)		
-	{
-		if (Key == EKeys::F11)
-		{
-			return false;
-		}
-	}
-
-	return FLevelEditorViewportClient::InputKey(Viewport, ControllerId, Key, Event, AmountDepressed, bGamepad);
-}
-
-
 class SPreArrangedBox : public SCompoundWidget
 {
 public:
@@ -125,16 +111,38 @@ private:
 	FOnArrange OnArrange;
 };
 
-class SCinematicPreviewViewport : public SEditorViewport
+class SCinematicPreviewViewport : public SLevelViewport
 {
 public:
+	virtual const FSlateBrush* OnGetViewportBorderBrush() const override { return nullptr; }
+	virtual EVisibility GetCurrentLevelTextVisibility() const override { return EVisibility::Collapsed; }
+	virtual EVisibility GetViewportControlsVisibility() const override { return EVisibility::Collapsed; }
+
+	virtual TSharedPtr<SWidget> MakeViewportToolbar() { return nullptr; }
+
+	TSharedPtr<SWidget> MakeExternalViewportToolbar() { return SLevelViewport::MakeViewportToolbar(); }
+
+	FSlateColor GetBorderColorAndOpacity() const
+	{
+		return OnGetViewportBorderColorAndOpacity();
+	}
+
+	const FSlateBrush* GetBorderBrush() const
+	{
+		return SLevelViewport::OnGetViewportBorderBrush();
+	}
 
-	virtual TSharedRef<FEditorViewportClient> MakeEditorViewportClient()
+	EVisibility GetBorderVisibility() const
 	{
-		return MakeShareable( new FCinematicViewportClient );
+		EVisibility Visibility = SLevelViewport::OnGetViewportContentVisibility();
+		return Visibility == EVisibility::Visible ? EVisibility::HitTestInvisible : Visibility;
 	}
+
+private:
+	bool bShowToolbar;
 };
 
+
 BEGIN_SLATE_FUNCTION_BUILD_OPTIMIZATION
 void SCinematicLevelViewport::Construct(const FArguments& InArgs)
 {
@@ -142,6 +150,17 @@ void SCinematicLevelViewport::Construct(const FArguments& InArgs)
 	LayoutName = InArgs._LayoutName;
 	RevertToLayoutName = InArgs._RevertToLayoutName;
 
+	ViewportClient = MakeShareable( new FCinematicViewportClient() );
+
+	ViewportWidget = SNew(SCinematicPreviewViewport)
+		.LevelEditorViewportClient(ViewportClient)
+		.ParentLevelEditor(InArgs._ParentLevelEditor)
+		.ParentLayout(ParentLayout.Pin())
+		.ConfigKey(LayoutName.ToString())
+		.Realtime(true);
+
+	ViewportClient->SetViewportWidget(ViewportWidget);
+
 	TypeInterfaceProxy = MakeShareable( new FTypeInterfaceProxy );
 
 	FLevelSequenceEditorToolkit::OnOpened().AddSP(this, &SCinematicLevelViewport::OnEditorOpened);
@@ -150,102 +169,75 @@ void SCinematicLevelViewport::Construct(const FArguments& InArgs)
 
 	TSharedRef<SFilmOverlayOptions> FilmOverlayOptions = SNew(SFilmOverlayOptions);
 
-	TSharedRef<SWidget> MaximizeButton = SNullWidget::NullWidget;
-	if (InArgs._bShowMaximize)
-	{
-		MaximizeButton =
-			SNew(SButton)
-			.Cursor(EMouseCursor::Default)
-			.ButtonStyle(FEditorStyle::Get(), "HoverHintOnly")
-			.OnClicked(this, &SCinematicLevelViewport::OnToggleMaximize)
-			.ToolTipText(LOCTEXT("Maximize_ToolTip", "Maximizes or restores this viewport"))
+	DecoratedTransportControls = SNew(SHorizontalBox)
+		+ SHorizontalBox::Slot()
+		[
+			SNew(SSpacer)
+		]
+
+		+ SHorizontalBox::Slot()
+		.AutoWidth()
+		.Padding(5.f, 0.f)
+		.VAlign(VAlign_Center)
+		.HAlign(HAlign_Center)
+		[
+			SNew(SBorder)
+			.Padding(0)
+			.BorderImage(nullptr)
+			.ForegroundColor(FEditorStyle::GetSlateColor("SelectionColor").GetColor(FWidgetStyle()))
 			[
-				SNew(SBox)
-				.VAlign(VAlign_Center)
-				.HAlign(HAlign_Center)
-				[
-					SNew(SImage)
-					.Image(this, &SCinematicLevelViewport::GetMaximizeImage)
-					.ColorAndOpacity(Gray)
-				]
-			];
-	}
+				SNew(SNonThrottledSpinBox<float>)
+				.TypeInterface(TypeInterfaceProxy)
+				.Style(FEditorStyle::Get(), "Sequencer.HyperlinkSpinBox")
+				.Font(FEditorStyle::GetFontStyle("Sequencer.FixedFont"))
+				.OnValueCommitted(this, &SCinematicLevelViewport::OnTimeCommitted)
+				.OnValueChanged(this, &SCinematicLevelViewport::SetTime)
+				.OnEndSliderMovement(this, &SCinematicLevelViewport::SetTime)
+				.MinValue(this, &SCinematicLevelViewport::GetMinTime)
+				.MaxValue(this, &SCinematicLevelViewport::GetMaxTime)
+				.Value(this, &SCinematicLevelViewport::GetTime)
+			]
+		]
 
-	ChildSlot
-	[
-		SNew(SBorder)
+		+ SHorizontalBox::Slot()
+		.AutoWidth()
+		.VAlign(VAlign_Center)
+		.HAlign(HAlign_Center)
+		[
+			SAssignNew(TransportControlsContainer, SBox)
+		]
+
+		+ SHorizontalBox::Slot()
+		[
+			SNew(SSpacer)
+		];
+
+	TSharedRef<SWidget> MainViewport = 	SNew(SBorder)
 		.BorderImage(FEditorStyle::GetBrush("BlackBrush"))
 		.ForegroundColor(Gray)
-		.Padding(5.f)
+		.Padding(0)
 		[
 			SNew(SVerticalBox)
 
 			+ SVerticalBox::Slot()
+			.Padding(5.f)
 			.AutoHeight()
 			[
 				SNew(SHorizontalBox)
 
 				+ SHorizontalBox::Slot()
-				.AutoWidth()
 				[
-					SNew(SComboButton)
-					.ButtonStyle(FEditorStyle::Get(), "HoverHintOnly")
-					.ForegroundColor(Gray)
-					.OnGetMenuContent(this, &SCinematicLevelViewport::MakeCameraMenu)
-					.ButtonContent()
-					[
-						SNew(SBox)
-						.WidthOverride(36)
-						.HeightOverride(24)
-						.ToolTipText(LOCTEXT("CameraSelectionTooltip", "Displays a list of available cameras to view through."))
-						[
-							SNew(SImage)
-							.Image(FLevelSequenceEditorStyle::Get()->GetBrush("LevelSequenceEditor.CinematicViewportCamera"))
-							.ColorAndOpacity(Gray)
-						]
-					]
-				]
-				
-				+ SHorizontalBox::Slot()
-				[
-					SNew(SSpacer)
-				]
-				
-				+ SHorizontalBox::Slot()
-				.AutoWidth()
-				[
-					FilmOverlayOptions
-				]
-
-				+ SHorizontalBox::Slot()
-				.AutoWidth()
-				[
-					MaximizeButton
+					ViewportWidget->MakeExternalViewportToolbar().ToSharedRef()
 				]
 
 				+ SHorizontalBox::Slot()
 				.AutoWidth()
 				[
-					SNew(SButton)
-					.Cursor(EMouseCursor::Default)
-					.ButtonStyle(FEditorStyle::Get(), "HoverHintOnly")
-					.OnClicked(this, &SCinematicLevelViewport::OnRevertLayout)
-					.ToolTipText(LOCTEXT("Close_ToolTip", "Close this cinematic viewport, reverting it back to a standard viewport"))
-					[
-						SNew(SBox)
-						.VAlign(VAlign_Center)
-						.HAlign(HAlign_Center)
-						[
-							SNew(SImage)
-							.Image(FLevelSequenceEditorStyle::Get()->GetBrush("LevelSequenceEditor.CinematicViewportClose"))
-							.ColorAndOpacity(Gray)
-						]
-					]
+					FilmOverlayOptions
 				]
 			]
 
 			+ SVerticalBox::Slot()
-			.Padding(20.f)
 			[
 				SNew(SPreArrangedBox)
 				.OnArrange(this, &SCinematicLevelViewport::CacheDesiredViewportSize)
@@ -269,7 +261,7 @@ void SCinematicLevelViewport::Construct(const FArguments& InArgs)
 
 							+ SOverlay::Slot()
 							[
-								SAssignNew(ViewportWidget, SCinematicPreviewViewport)
+								ViewportWidget.ToSharedRef()
 							]
 
 							+ SOverlay::Slot()
@@ -282,11 +274,12 @@ void SCinematicLevelViewport::Construct(const FArguments& InArgs)
 					+ SVerticalBox::Slot()
 					.AutoHeight()
 					.HAlign(HAlign_Center)
+					.Padding(10.f, 0.f)
 					[
 						SAssignNew(ViewportControls, SBox)
 						.Visibility(this, &SCinematicLevelViewport::GetControlsVisibility)
 						.WidthOverride(this, &SCinematicLevelViewport::GetDesiredViewportWidth)
-						.Padding(FMargin(0.f, 10.f))
+						.Padding(FMargin(0.f, 10.f, 0.f, 0.f))
 						[
 							SNew(SHorizontalBox)
 
@@ -317,7 +310,7 @@ void SCinematicLevelViewport::Construct(const FArguments& InArgs)
 							]
 						]
 					]
-
+				
 					+ SVerticalBox::Slot()
 					[
 						SNew(SSpacer)
@@ -325,11 +318,10 @@ void SCinematicLevelViewport::Construct(const FArguments& InArgs)
 				]
 			]
 
-
 			+ SVerticalBox::Slot()
+			.Padding(5.f)
 			.AutoHeight()
 			[
-			
 				SNew(SWidgetSwitcher)
 				.WidgetIndex(this, &SCinematicLevelViewport::GetVisibleWidgetIndex)
 
@@ -348,95 +340,42 @@ void SCinematicLevelViewport::Construct(const FArguments& InArgs)
 					.AutoHeight()
 					.Padding(5.f, 0.f)
 					[
-						SNew(SHorizontalBox)
-
-						+ SHorizontalBox::Slot()
-						.VAlign(VAlign_Center)
-						.AutoWidth()
-						[
-							SNew(STextBlock)
-							.Font(FEditorStyle::GetFontStyle("Sequencer.FixedFont"))
-							.ColorAndOpacity(Gray)
-							.Text_Lambda([=]{ return UIData.MasterStartText; })
-						]
-
-						+ SHorizontalBox::Slot()
-						.VAlign(VAlign_Center)
-						.AutoWidth()
-						.Padding(FMargin(10.f, 0.f))
-						[
-							SNew(SNonThrottledSpinBox<float>)
-							.TypeInterface(TypeInterfaceProxy)
-							.Style(FEditorStyle::Get(), "Sequencer.HyperlinkSpinBox")
-							.Font(FEditorStyle::GetFontStyle("Sequencer.FixedFont"))
-							.OnValueCommitted(this, &SCinematicLevelViewport::OnTimeCommitted)
-							.OnValueChanged(this, &SCinematicLevelViewport::SetTime)
-							.OnEndSliderMovement(this, &SCinematicLevelViewport::SetTime)
-							.MinValue(this, &SCinematicLevelViewport::GetMinTime)
-							.MaxValue(this, &SCinematicLevelViewport::GetMaxTime)
-							.Value(this, &SCinematicLevelViewport::GetTime)
-						]
-
-						+ SHorizontalBox::Slot()
-						.VAlign(VAlign_Center)
-						.HAlign(HAlign_Right)
-						[
-							SNew(SButton)
-							.OnClicked(this, &SCinematicLevelViewport::SetPlaybackStart)
-							.ToolTipText(LOCTEXT("SetPlayStart_Tooltip", "Set playback start to the current position"))
-							.ButtonStyle(FEditorStyle::Get(), "NoBorder")
-							.ContentPadding(2.0f)
-							[
-								SNew(SImage)
-								.Image(FLevelSequenceEditorStyle::Get()->GetBrush("LevelSequenceEditor.CinematicViewportSetPlayStart"))
-							]
-						]
-
-						+ SHorizontalBox::Slot()
-						.VAlign(VAlign_Center)
-						.HAlign(HAlign_Center)
-						.AutoWidth()
-						[
-							SAssignNew(TransportControls, SBox)
-						]
-
-						+ SHorizontalBox::Slot()
-						.VAlign(VAlign_Center)
-						.HAlign(HAlign_Left)
-						[
-							SNew(SButton)
-							.OnClicked(this, &SCinematicLevelViewport::SetPlaybackEnd)
-							.ToolTipText(LOCTEXT("SetPlayEnd_Tooltip", "Set playback end to the current position"))
-							.ButtonStyle(FEditorStyle::Get(), "NoBorder")
-							.ContentPadding(2.0f)
-							[
-								SNew(SImage)
-								.Image(FLevelSequenceEditorStyle::Get()->GetBrush("LevelSequenceEditor.CinematicViewportSetPlayEnd"))
-							]
-						]
-
-						+ SHorizontalBox::Slot()
-						.VAlign(VAlign_Center)
-						.AutoWidth()
-						[
-							SNew(STextBlock)
-							.Font(FEditorStyle::GetFontStyle("Sequencer.FixedFont"))
-							.ColorAndOpacity(Gray)
-							.Text_Lambda([=]{ return UIData.MasterEndText; })
-						]
+						SAssignNew(TimeRangeContainer, SBox)
 					]
 				]
 
 				+ SWidgetSwitcher::Slot()
 				.VAlign(VAlign_Center)
 				.HAlign(HAlign_Center)
-				.Padding(5.f, 0.f)
 				[
-					SNew(STextBlock)
-					.ColorAndOpacity(Gray)
-					.Text(LOCTEXT("NoSequencerMessage", "No active Level Sequencer detected. Please edit a Level Sequence to enable full controls."))
+					SNew(SBox)
+					.Padding(FMargin(5.f, 10.f))
+					[
+						SNew(STextBlock)
+						.ColorAndOpacity(Gray)
+						.Text(LOCTEXT("NoSequencerMessage", "No active Level Sequencer detected. Please edit a Level Sequence to enable full controls."))
+					]
 				]
 			]
+		];
+
+	ChildSlot
+	[
+		SNew(SOverlay)
+
+		+ SOverlay::Slot()
+		[
+			MainViewport
+		]
+
+		+SOverlay::Slot()
+		[
+			SNew(SBorder)
+			.BorderImage(ViewportWidget.Get(), &SCinematicPreviewViewport::GetBorderBrush)
+			.BorderBackgroundColor(ViewportWidget.Get(), &SCinematicPreviewViewport::GetBorderColorAndOpacity)
+			.Visibility(ViewportWidget.Get(), &SCinematicPreviewViewport::GetBorderVisibility)
+			.Padding(0.0f)
+			.ShowEffectWhenDisabled( false )
 		]
 	];
 
@@ -448,25 +387,12 @@ void SCinematicLevelViewport::Construct(const FArguments& InArgs)
 	CommandList = MakeShareable( new FUICommandList );
 	// Ensure the commands are registered
 	FLevelSequenceEditorCommands::Register();
-	BindCommands();
 }
 END_SLATE_FUNCTION_BUILD_OPTIMIZATION
 
-FReply SCinematicLevelViewport::OnRevertLayout()
+TSharedPtr<SLevelViewport> SCinematicLevelViewport::GetLevelViewport() const
 {
-	TSharedPtr<FLevelViewportLayout> LayoutPinned = ParentLayout.Pin();
-	if (LayoutPinned.IsValid())
-	{
-		TSharedPtr<FLevelViewportTabContent> ViewportTabPinned = LayoutPinned->GetParentTabContent().Pin();
-		if (ViewportTabPinned.IsValid())
-		{
-			// Viewport clients are going away.  Any current one is invalid.
-			GCurrentLevelEditingViewportClient = nullptr;
-			ViewportTabPinned->SetViewportConfiguration(RevertToLayoutName);
-			FSlateApplication::Get().DismissAllMenus();
-		}
-	}
-	return FReply::Handled();
+	return ViewportWidget;
 }
 
 int32 SCinematicLevelViewport::GetVisibleWidgetIndex() const
@@ -476,58 +402,7 @@ int32 SCinematicLevelViewport::GetVisibleWidgetIndex() const
 
 EVisibility SCinematicLevelViewport::GetControlsVisibility() const
 {
-	return CurrentToolkit.IsValid() ? EVisibility::Visible : EVisibility::Hidden;
-}
-
-FReply SCinematicLevelViewport::OnToggleMaximize()
-{
-	TSharedPtr<FLevelViewportLayout> ParentLayoutPinned = ParentLayout.Pin();
-	if (ParentLayoutPinned.IsValid() && ParentLayoutPinned->IsMaximizeSupported())
-	{
-		const bool bAllowAnimation = true;
-		const bool bMaximize = !ParentLayoutPinned->IsViewportMaximized(LayoutName);
-		
-		ParentLayoutPinned->RequestMaximizeViewport( LayoutName, bMaximize, false, bAllowAnimation );
-	}
-	return FReply::Handled();
-}
-
-bool SCinematicLevelViewport::IsMaximized() const
-{
-	TSharedPtr<FLevelViewportLayout> ParentLayoutPinned = ParentLayout.Pin();
-	if (ParentLayoutPinned.IsValid())
-	{
-		return ParentLayoutPinned->IsViewportMaximized(LayoutName);
-	}
-	return false;
-}
-
-void SCinematicLevelViewport::OnToggleImmersive()
-{
-	TSharedPtr<FLevelViewportLayout> ParentLayoutPinned = ParentLayout.Pin();
-	if (ParentLayoutPinned.IsValid() && ParentLayoutPinned->IsMaximizeSupported())
-	{
-		const bool bAllowAnimation = true;
-		const bool bMaximize = !ParentLayoutPinned->IsViewportMaximized(LayoutName);
-		
-		const bool bImmersive = !IsImmersive();
-		ParentLayoutPinned->RequestMaximizeViewport( LayoutName, bMaximize, bImmersive, bAllowAnimation );
-	}
-}
-
-bool SCinematicLevelViewport::IsImmersive() const
-{
-	TSharedPtr<FLevelViewportLayout> ParentLayoutPinned = ParentLayout.Pin();
-	if (ParentLayoutPinned.IsValid())
-	{
-		return ParentLayoutPinned->IsViewportImmersive(LayoutName);
-	}
-	return false;
-}
-
-const FSlateBrush* SCinematicLevelViewport::GetMaximizeImage() const
-{
-	return IsMaximized() ? FEditorStyle::GetBrush("LevelViewportToolBar.Maximize.Checked") : FEditorStyle::GetBrush("LevelViewportToolBar.Maximize.Normal");
+	return CurrentToolkit.IsValid() ? EVisibility::Visible : EVisibility::Collapsed;
 }
 
 TOptional<float> SCinematicLevelViewport::GetMinTime() const
@@ -574,46 +449,15 @@ float SCinematicLevelViewport::GetTime() const
 	return 0.f;
 }
 
-FReply SCinematicLevelViewport::SetPlaybackStart()
-{
-	ISequencer* Sequencer = GetSequencer();
-	if (Sequencer)
-	{
-		FScopedTransaction Transaction(LOCTEXT("SetPlaybackStart", "Set Playback Start"));
-
-		TRange<float> CurrentRange = Sequencer->GetFocusedMovieSceneSequence()->GetMovieScene()->GetPlaybackRange();
-
-		float NewPos = FMath::Min(Sequencer->GetGlobalTime(), CurrentRange.GetUpperBoundValue());
-		Sequencer->GetFocusedMovieSceneSequence()->GetMovieScene()->SetPlaybackRange(NewPos, CurrentRange.GetUpperBoundValue());
-	}
-	return FReply::Handled();
-}
-FReply SCinematicLevelViewport::SetPlaybackEnd()
-{
-	ISequencer* Sequencer = GetSequencer();
-	if (Sequencer)
-	{
-		FScopedTransaction Transaction(LOCTEXT("SetPlaybackEnd", "Set Playback End"));
-
-		TRange<float> CurrentRange = Sequencer->GetFocusedMovieSceneSequence()->GetMovieScene()->GetPlaybackRange();
-
-		float NewPos = FMath::Max(Sequencer->GetGlobalTime(), CurrentRange.GetLowerBoundValue());
-		Sequencer->GetFocusedMovieSceneSequence()->GetMovieScene()->SetPlaybackRange(CurrentRange.GetLowerBoundValue(), NewPos);
-	}
-	return FReply::Handled();
-}
 void SCinematicLevelViewport::CacheDesiredViewportSize(const FGeometry& AllottedGeometry)
 {
-	FVector2D ViewportControlsSize = ViewportControls->GetDesiredSize();
-
 	FVector2D AllowableSpace = AllottedGeometry.GetLocalSize();
-	AllowableSpace.Y -= ViewportControlsSize.Y;
+	AllowableSpace.Y -= ViewportControls->GetDesiredSize().Y;
 
-	TSharedPtr<FCinematicViewportClient> ViewportClient = GetViewportClient();
 	if (ViewportClient->IsAspectRatioConstrained())
 	{
-		const float MinSize = FMath::Min(AllowableSpace.X / ViewportClient->AspectRatio, AllowableSpace.Y);
-		DesiredViewportSize = FVector2D(int32(ViewportClient->AspectRatio * MinSize), int32(MinSize));
+		const float MinSize = FMath::TruncToFloat(FMath::Min(AllowableSpace.X / ViewportClient->AspectRatio, AllowableSpace.Y));
+		DesiredViewportSize = FVector2D(FMath::TruncToFloat(ViewportClient->AspectRatio * MinSize), MinSize);
 	}
 	else
 	{
@@ -631,214 +475,6 @@ FOptionalSize SCinematicLevelViewport::GetDesiredViewportHeight() const
 	return DesiredViewportSize.Y;
 }
 
-TSharedRef<SWidget> SCinematicLevelViewport::MakeCameraMenu()
-{
-	ISequencer* Sequencer = GetSequencer();
-	UMovieSceneSequence* Sequence = Sequencer ? Sequencer->GetFocusedMovieSceneSequence() : nullptr;
-
-	FMenuBuilder MenuBuilder(true, nullptr, nullptr);
-	{
-		// let toolkits populate the menu
-		MenuBuilder.AddMenuEntry(
-			LOCTEXT("FreeCamera", "Free Camera"),
-			LOCTEXT("FreeCamera_ToolTip", "Allows the arbitrary movement of the camera"),
-			FSlateIcon(),
-			FUIAction(
-				FExecuteAction::CreateSP(this, &SCinematicLevelViewport::SetCameraMode, ECinematicCameraMode::FreeCamera, TOptional<FGuid>()),
-				FCanExecuteAction::CreateLambda([=]{ return true; }),
-				FIsActionChecked::CreateLambda([=]{ return CameraMode == ECinematicCameraMode::FreeCamera; })
-			),
-			NAME_None,
-			EUserInterfaceActionType::ToggleButton
-		);
-
-		if (Sequence)
-		{
-			MenuBuilder.BeginSection("Sequence");
-			{
-				MenuBuilder.AddMenuEntry(
-					LOCTEXT("AutoCamera", "Auto"),
-					LOCTEXT("AutoCamera_ToolTip", "Looks through the camera defined by either the shot track, camera cut track, or any other available camera (in that order)"),
-					FSlateIcon(),
-					FUIAction(
-						FExecuteAction::CreateSP(this, &SCinematicLevelViewport::SetCameraMode, ECinematicCameraMode::Auto, TOptional<FGuid>()),
-						FCanExecuteAction::CreateLambda([=]{ return true; }),
-						FIsActionChecked::CreateLambda([=]{ return CameraMode == ECinematicCameraMode::Auto; })
-					),
-					NAME_None,
-					EUserInterfaceActionType::ToggleButton
-				);
-			}
-			MenuBuilder.EndSection();
-
-			MenuBuilder.BeginSection("Cameras");
-			{
-				for (auto& Pair : Cameras)
-				{
-					FGuid CameraGuid = Pair.Key;
-
-					MenuBuilder.AddMenuEntry(
-						FText::FromString(Pair.Value.DisplayName),
-						FText(),
-						FSlateIcon(),
-						FUIAction(
-							FExecuteAction::CreateSP(this, &SCinematicLevelViewport::SetCameraMode, ECinematicCameraMode::SpecificCamera, TOptional<FGuid>(CameraGuid)),
-							FCanExecuteAction::CreateLambda([=]{ return true; }),
-							FIsActionChecked::CreateLambda([=]{
-								return CameraMode == ECinematicCameraMode::SpecificCamera &&
-									LockedCameraGuid.IsSet() &&
-									LockedCameraGuid.GetValue() == CameraGuid;
-							})
-						),
-						NAME_None,
-						EUserInterfaceActionType::ToggleButton
-					);
-				}
-			}
-			MenuBuilder.EndSection();
-		}
-	}
-
-	return MenuBuilder.MakeWidget();
-}
-
-void SCinematicLevelViewport::SetCameraMode(ECinematicCameraMode InCameraMode, TOptional<FGuid> InSpecificCameraGuid)
-{
-	CameraMode = InCameraMode;
-
-	TSharedPtr<FCinematicViewportClient> ViewportClient = GetViewportClient();
-
-	// Reset everything
-	LockedCameraGuid = TOptional<FGuid>();
-	ViewportClient->SetActorLock(nullptr);
-
-	switch (CameraMode)
-	{
-	case ECinematicCameraMode::FreeCamera:
-		ViewportClient->bDisableInput = false;
-		break;
-
-	case ECinematicCameraMode::Auto:
-		ViewportClient->bDisableInput = true;
-		break;
-
-	case ECinematicCameraMode::SpecificCamera:
-		LockedCameraGuid = InSpecificCameraGuid;
-		ViewportClient->bDisableInput = true;
-		break;
-	}
-
-	ISequencer* Sequencer = GetSequencer();
-	if (Sequencer)
-	{
-		Sequencer->SetGlobalTime(Sequencer->GetGlobalTime());
-	}
-}
-
-void SCinematicLevelViewport::UpdateActiveCameras()
-{
-	Cameras.Reset();
-
-	ISequencer* Sequencer = GetSequencer();
-	if (!Sequencer)
-	{
-		return;
-	}
-
-	TSharedPtr<FMovieSceneSequenceInstance> SequenceInstance = Sequencer->GetFocusedMovieSceneSequenceInstance();
-	UMovieScene* MovieScene = SequenceInstance->GetSequence()->GetMovieScene();
-
-	const int32 NumSpawnables = MovieScene->GetSpawnableCount();
-	for (int32 SpawnableIndex = 0; SpawnableIndex < NumSpawnables; ++SpawnableIndex)
-	{
-		FMovieSceneSpawnable& Spawnable = MovieScene->GetSpawnable(SpawnableIndex);
-		if (Spawnable.GetClass()->IsChildOf(AActor::StaticClass()))
-		{
-			AActor* CameraActor = Cast<AActor>(SequenceInstance->FindObject(Spawnable.GetGuid(), *Sequencer));
-			if (CameraActor)
-			{
-				UCameraComponent* CameraComponent = MovieSceneHelpers::CameraComponentFromActor(CameraActor);
-				if (CameraComponent)
-				{
-					Cameras.Add(Spawnable.GetGuid(), FActiveCamera(Spawnable.GetName(), CameraActor));
-				}
-			}
-		}
-	}
-
-	const int32 NumPossessables = MovieScene->GetPossessableCount();
-	for (int32 PossessableIndex = 0; PossessableIndex < NumPossessables; ++PossessableIndex)
-	{
-		FMovieScenePossessable& Possessable = MovieScene->GetPossessable(PossessableIndex);
-		AActor* CameraActor = Cast<AActor>(SequenceInstance->FindObject(Possessable.GetGuid(), *Sequencer));
-		if (CameraActor)
-		{
-			UCameraComponent* CameraComponent = MovieSceneHelpers::CameraComponentFromActor(CameraActor);
-			if (CameraComponent)
-			{
-				Cameras.Add(Possessable.GetGuid(), FActiveCamera(Possessable.GetName(), CameraActor));
-			}
-		}
-	}
-}
-
-TSharedPtr<FCinematicViewportClient> SCinematicLevelViewport::GetViewportClient() const
-{
-	return StaticCastSharedPtr<FCinematicViewportClient>(ViewportWidget->GetViewportClient());
-}
-
-void SCinematicLevelViewport::OnUpdateCameraCut(UObject* CameraObject)
-{
-	if (CameraMode != ECinematicCameraMode::Auto)
-	{
-		return;
-	}
-
-	AActor* CameraActor = Cast<AActor>(CameraObject);
-	TSharedPtr<FCinematicViewportClient> ViewportClient = GetViewportClient();
-	
-	if (CameraActor)
-	{
-		ViewportClient->SetViewLocation(CameraActor->GetActorLocation());
-		ViewportClient->SetViewRotation(CameraActor->GetActorRotation());
-		//ViewportClient->bEditorCameraCut = !ViewportClient->IsLockedToActor(CameraActor);
-		ViewportClient->bLockedCameraView = true;
-	}
-	else
-	{
-		ViewportClient->ViewFOV = ViewportClient->FOVAngle;
-		//ViewportClient->bEditorCameraCut = false;
-		ViewportClient->bLockedCameraView = false;
-	}
-
-	// Set the actor lock.
-	ViewportClient->SetActorLock(CameraActor);
-	ViewportClient->RemoveCameraRoll();
-
-	// If viewing through a camera - enforce aspect ratio.
-	if (CameraActor)
-	{
-		UCameraComponent* CameraComponent = MovieSceneHelpers::CameraComponentFromActor(CameraActor);
-		if (CameraComponent)
-		{
-			if (CameraComponent->AspectRatio == 0)
-			{
-				ViewportClient->AspectRatio = 1.7f;
-			}
-			else
-			{
-				ViewportClient->AspectRatio = CameraComponent->AspectRatio;
-			}
-
-			//don't stop the camera from zooming when not playing back
-			ViewportClient->ViewFOV = CameraComponent->FieldOfView;
-		}
-	}
-
-	// Update ControllingActorViewInfo, so it is in sync with the updated viewport
-	ViewportClient->UpdateViewForLockedActor();
-}
-
 FReply SCinematicLevelViewport::OnKeyDown( const FGeometry& MyGeometry, const FKeyEvent& InKeyEvent ) 
 {
 	if (CommandList->ProcessCommandBindings(InKeyEvent))
@@ -864,37 +500,32 @@ void SCinematicLevelViewport::Setup(FLevelSequenceEditorToolkit& NewToolkit)
 	ISequencer* Sequencer = GetSequencer();
 	if (Sequencer)
 	{
+		Sequencer->SetViewportTransportControlsVisibility(false);
 		TypeInterfaceProxy->Impl = Sequencer->GetZeroPadNumericTypeInterface();
 
-		OnCameraCutHandle = Sequencer->OnCameraCut().AddSP(this, &SCinematicLevelViewport::OnUpdateCameraCut);
-
 		if (TransportRange.IsValid())
 		{
 			TransportRange->SetSequencer(Sequencer->AsShared());
 		}
 
-		if (TransportControls.IsValid())
+		if (TransportControlsContainer.IsValid())
 		{
-			TransportControls->SetContent(Sequencer->MakeTransportControls());
+			TransportControlsContainer->SetContent(Sequencer->MakeTransportControls(true));
 		}
 
-		SetCameraMode(ECinematicCameraMode::Auto);
-	}
-	else
-	{
-		SetCameraMode(ECinematicCameraMode::FreeCamera);
+		if (TimeRangeContainer.IsValid())
+		{
+			const bool bShowWorkingRange = true, bShowViewRange = false, bShowPlaybackRange = true;
+			TimeRangeContainer->SetContent(Sequencer->MakeTimeRange(DecoratedTransportControls.ToSharedRef(), bShowWorkingRange, bShowViewRange, bShowPlaybackRange));
+		}
 	}
 }
 
 void SCinematicLevelViewport::CleanUp()
 {
-	ISequencer* Sequencer = GetSequencer();
-	if (Sequencer && OnCameraCutHandle.IsValid())
-	{
-		Sequencer->OnCameraCut().Remove(OnCameraCutHandle);
-	}
+	TransportControlsContainer->SetContent(SNullWidget::NullWidget);
+	TimeRangeContainer->SetContent(SNullWidget::NullWidget);
 
-	TransportControls->SetContent(SNullWidget::NullWidget);
 }
 
 void SCinematicLevelViewport::OnEditorOpened(FLevelSequenceEditorToolkit& Toolkit)
@@ -919,23 +550,6 @@ void SCinematicLevelViewport::OnEditorClosed()
 	{
 		Setup(*NewToolkit);
 	}
-	else
-	{
-		SetCameraMode(ECinematicCameraMode::FreeCamera);
-	}
-}
-
-void SCinematicLevelViewport::BindCommands()
-{
-	FUICommandList& CommandListRef = *CommandList;
-
-	const FLevelSequenceEditorCommands& Commands = FLevelSequenceEditorCommands::Get();
-
-	CommandListRef.MapAction( 
-		Commands.ToggleImmersive,
-		FExecuteAction::CreateSP( this, &SCinematicLevelViewport::OnToggleImmersive ),
-		FCanExecuteAction(),
-		FIsActionChecked::CreateSP(this, &SCinematicLevelViewport::IsImmersive));
 }
 
 ISequencer* SCinematicLevelViewport::GetSequencer() const
@@ -954,35 +568,17 @@ void SCinematicLevelViewport::Tick(const FGeometry& AllottedGeometry, const doub
 	SCompoundWidget::Tick(AllottedGeometry, InCurrentTime, InDeltaTime);
 
 	ISequencer* Sequencer = GetSequencer();
-	UMovieSceneSequence* Sequence = Sequencer ? Sequencer->GetFocusedMovieSceneSequence() : nullptr;
-
-	if (!Sequence)
+	if (!Sequencer)
 	{
 		return;
 	}
 
-	UpdateActiveCameras();
+	ViewportClient->bDisableInput = Sequencer->GetPlaybackStatus() == EMovieScenePlayerStatus::Playing;
 
-	TSharedPtr<FCinematicViewportClient> ViewportClient = GetViewportClient();
-
-	// Update the specific locked camera if set. We do this each tick as it might be a spawnable camera
-	if (LockedCameraGuid.IsSet())
+	UMovieSceneSequence* Sequence = Sequencer->GetFocusedMovieSceneSequence();
+	if (!Sequence)
 	{
-		FActiveCamera* Camera = Cameras.Find(LockedCameraGuid.GetValue());
-		if (Camera && Camera->CameraActor.IsValid())
-		{
-			ViewportClient->SetActorLock(Camera->CameraActor.Get());
-
-			ViewportClient->bLockedCameraView = true;
-			ViewportClient->bDisableInput = true;
-		}
-		else
-		{
-			// If the camera's not available, just use free cam until it is
-			ViewportClient->SetActorLock(nullptr);
-			ViewportClient->bLockedCameraView = false;
-			ViewportClient->bDisableInput = false;
-		}
+		return;
 	}
 
 	// Find either a cinematic shot track or a sub track
@@ -1033,7 +629,7 @@ void SCinematicLevelViewport::Tick(const FGeometry& AllottedGeometry, const doub
 		{
 			UIData.ShotName = CinematicShotSection->GetShotDisplayName();
 		}
-		else
+		else if (SubSection->GetSequence() != nullptr)
 		{
 			UIData.ShotName = SubSection->GetSequence()->GetDisplayName();
 		}
diff --git a/Engine/Plugins/MovieScene/LevelSequenceEditor/Source/LevelSequenceEditor/Private/CinematicViewport/SCinematicLevelViewport.h b/Engine/Plugins/MovieScene/LevelSequenceEditor/Source/LevelSequenceEditor/Private/CinematicViewport/SCinematicLevelViewport.h
--- a/Engine/Plugins/MovieScene/LevelSequenceEditor/Source/LevelSequenceEditor/Private/CinematicViewport/SCinematicLevelViewport.h
+++ b/Engine/Plugins/MovieScene/LevelSequenceEditor/Source/LevelSequenceEditor/Private/CinematicViewport/SCinematicLevelViewport.h
@@ -3,12 +3,13 @@
 #pragma once
 
 #include "LevelEditorViewport.h"
-#include "SEditorViewport.h"
+#include "SLevelViewport.h"
 
 
 class FLevelSequenceEditorToolkit;
 class ISequencer;
 class SCinematicTransportRange;
+class SCinematicPreviewViewport;
 class FLevelViewportLayout;
 struct FTypeInterfaceProxy;
 
@@ -17,8 +18,7 @@ struct FTypeInterfaceProxy;
 struct FCinematicViewportClient : FLevelEditorViewportClient
 {
 	FCinematicViewportClient();
-	virtual bool IsLevelEditorClient() const override { return true; }
-	virtual bool InputKey(FViewport* Viewport, int32 ControllerId, FKey Key, EInputEvent Event, float AmountDepressed = 1.f, bool bGamepad=false) override;
+	void SetViewportWidget(const TSharedPtr<SEditorViewport>& InViewportWidget) { EditorViewportWidget = InViewportWidget; }
 };
 
 
@@ -40,20 +40,11 @@ struct FUIData
 };
 
 
-/** Enum specifying different camera modes */
-enum class ECinematicCameraMode
-{
-	FreeCamera,
-	Auto,
-	SpecificCamera
-};
-
-
 class SCinematicLevelViewport : public SCompoundWidget
 {
 public:
 
-	SLATE_BEGIN_ARGS(SCinematicLevelViewport) : _bShowMaximize(true) {}
+	SLATE_BEGIN_ARGS(SCinematicLevelViewport) {}
 		
 		/** The unique name of this viewport inside its parent layout */
 		SLATE_ARGUMENT(FName, LayoutName)
@@ -61,12 +52,14 @@ public:
 		SLATE_ARGUMENT(FName, RevertToLayoutName)
 		/** Ptr to this viewport's parent layout */
 		SLATE_ARGUMENT(TSharedPtr<FLevelViewportLayout>, ParentLayout)
-		/** Whether to show the maximize button or not */
-		SLATE_ARGUMENT(bool, bShowMaximize)
+		/** Ptr to this viewport's parent level editor */
+		SLATE_ARGUMENT(TWeakPtr<ILevelEditor>, ParentLevelEditor)
 	SLATE_END_ARGS()
 
 	/** Access this viewport's viewport client */
-	TSharedPtr<FCinematicViewportClient> GetViewportClient() const;
+	TSharedPtr<FCinematicViewportClient> GetViewportClient() const { return ViewportClient; }
+
+	TSharedPtr<SLevelViewport> GetLevelViewport() const;
 
 	/** Construct this widget */
 	void Construct(const FArguments& InArgs);
@@ -88,44 +81,13 @@ private:
 	/** Called when the level sequence editor toolkit we were observing has been closed */
 	void OnEditorClosed();
 
-	void BindCommands();
-
 private:
 
 	/** Get the sequencer ptr from our current toolkit, if set */
 	ISequencer* GetSequencer() const;
 
-	/** Called when our sequencer wants to switch cameras */
-	void OnUpdateCameraCut(UObject* CameraObject);
-
-	/** Set the camera to the specified mode (optionally binding to a specific camera) */
-	void SetCameraMode(ECinematicCameraMode InCameraMode, TOptional<FGuid> InSpecificCameraGuid = TOptional<FGuid>());
-
-	/** Make the currently available camera menu */
-	TSharedRef<SWidget> MakeCameraMenu();
-
-private:
-
-	/** Toggle this viewport's maximized state */
-	FReply OnToggleMaximize();
-
-	/** Test whether this viewport is maximized or not */
-	bool IsMaximized() const;
-
-	/** Toggle this viewport's immersive state */
-	void OnToggleImmersive();
-
-	/** Test whether this viewport is immersive or not */
-	bool IsImmersive() const;
-
-	/** Called when we want to revert to the equivalent non-cinematic viewport */
-	FReply OnRevertLayout();
-
 private:
 
-	/** Called once per frame to update the available cameras in the scene */
-	void UpdateActiveCameras();
-
 	/** Cache the desired viewport size based on the filled, allotted geometry */
 	void CacheDesiredViewportSize(const FGeometry& AllottedGeometry);
 
@@ -133,9 +95,6 @@ private:
 	FOptionalSize GetDesiredViewportWidth() const;
 	FOptionalSize GetDesiredViewportHeight() const;
 
-	/** Get the image for the maximize button */
-	const FSlateBrush* GetMaximizeImage() const;
-
 	int32 GetVisibleWidgetIndex() const;
 	EVisibility GetControlsVisibility() const;
 
@@ -151,14 +110,10 @@ private:
 
 	float GetTime() const;
 
-	FReply SetPlaybackStart();
-
-	FReply SetPlaybackEnd();
-
 private:
 	
 	/** Widget where the scene viewport is drawn in */
-	TSharedPtr<SEditorViewport> ViewportWidget;
+	TSharedPtr<SCinematicPreviewViewport> ViewportWidget;
 
 	/** The toolkit we're currently editing */
 	TWeakPtr<FLevelSequenceEditorToolkit> CurrentToolkit;
@@ -167,7 +122,10 @@ private:
 	TSharedPtr<FUICommandList> CommandList;
 
 	/** Slot for transport controls */
-	TSharedPtr<SBox> TransportControls;
+	TSharedPtr<SBox> TransportControlsContainer, TimeRangeContainer;
+
+	/** Decorated transport controls */
+	TSharedPtr<SWidget> DecoratedTransportControls;
 
 	/** Transport range widget */
 	TSharedPtr<SCinematicTransportRange> TransportRange;
@@ -175,15 +133,6 @@ private:
 	/** Cached UI data */
 	FUIData UIData;
 
-	/** Delegate binding handle for ISequencer::OnCameraCut */
-	FDelegateHandle OnCameraCutHandle;
-
-	/** The current camera mode */
-	ECinematicCameraMode CameraMode;
-
-	/** Guid of the camera in the UMovieScene to look through. Stored as a GUID rather than a ptr to ensure spawnable cameras work. */
-	TOptional<FGuid> LockedCameraGuid;
-
 	/** Cached desired size of the viewport */
 	FVector2D DesiredViewportSize;
 	
@@ -201,15 +150,6 @@ private:
 	/** The time spin box */
 	TSharedPtr<FTypeInterfaceProxy> TypeInterfaceProxy;
 
-	/** Cached active cameras */
-	struct FActiveCamera
-	{
-		FActiveCamera(FString InDisplayName, AActor* InCameraActor)
-			: DisplayName(InDisplayName), CameraActor(InCameraActor)
-		{}
-
-		FString DisplayName;
-		TWeakObjectPtr<AActor> CameraActor;
-	};
-	TMap<FGuid, FActiveCamera> Cameras;
+	/** The level editor viewport client for this viewport */
+	TSharedPtr<FCinematicViewportClient> ViewportClient;
 };
\ No newline at end of file
diff --git a/Engine/Plugins/MovieScene/LevelSequenceEditor/Source/LevelSequenceEditor/Private/CinematicViewport/SCinematicTransportRange.cpp b/Engine/Plugins/MovieScene/LevelSequenceEditor/Source/LevelSequenceEditor/Private/CinematicViewport/SCinematicTransportRange.cpp
--- a/Engine/Plugins/MovieScene/LevelSequenceEditor/Source/LevelSequenceEditor/Private/CinematicViewport/SCinematicTransportRange.cpp
+++ b/Engine/Plugins/MovieScene/LevelSequenceEditor/Source/LevelSequenceEditor/Private/CinematicViewport/SCinematicTransportRange.cpp
@@ -35,6 +35,13 @@ FReply SCinematicTransportRange::OnMouseButtonDown(const FGeometry& MyGeometry,
 {
 	bDraggingTime = true;
 	SetTime(MyGeometry, MouseEvent);
+		
+	ISequencer* Sequencer = GetSequencer();
+	if (Sequencer)
+	{
+		Sequencer->SetPlaybackStatus(EMovieScenePlayerStatus::Scrubbing);
+	}
+	
 	return FReply::Handled().CaptureMouse(AsShared()).PreventThrottling();
 }
 
@@ -51,6 +58,13 @@ FReply SCinematicTransportRange::OnMouseMove(const FGeometry& MyGeometry, const
 FReply SCinematicTransportRange::OnMouseButtonUp(const FGeometry& MyGeometry, const FPointerEvent& MouseEvent)
 {
 	bDraggingTime = false;
+
+	ISequencer* Sequencer = GetSequencer();
+	if (Sequencer)
+	{
+		Sequencer->SetPlaybackStatus(EMovieScenePlayerStatus::Stepping);
+	}
+	
 	return FReply::Handled().ReleaseMouseCapture();
 }
 
@@ -63,8 +77,8 @@ void SCinematicTransportRange::SetTime(const FGeometry& MyGeometry, const FPoint
 		Lerp = FMath::Clamp(Lerp, 0.f, 1.f);
 		
 		const TRange<float> WorkingRange = Sequencer->GetFocusedMovieSceneSequence()->GetMovieScene()->GetEditorData().WorkingRange;
-		const bool bAllowSnappingToFrames = true;
-		Sequencer->SetGlobalTime(WorkingRange.GetLowerBoundValue() + WorkingRange.Size<float>()*Lerp, bAllowSnappingToFrames);
+		
+		Sequencer->SetGlobalTime(WorkingRange.GetLowerBoundValue() + WorkingRange.Size<float>()*Lerp, ESnapTimeMode::STM_All);
 	}
 }
 
@@ -73,6 +87,15 @@ void SCinematicTransportRange::OnMouseCaptureLost()
 	bDraggingTime = false;
 }
 
+void SCinematicTransportRange::Tick(const FGeometry& AllottedGeometry, const double InCurrentTime, const float InDeltaTime)
+{
+	ISequencer* Sequencer = GetSequencer();
+	if (Sequencer)
+	{
+		Sequencer->GetKeysFromSelection(ActiveKeyCollection);
+	}
+}
+
 int32 SCinematicTransportRange::OnPaint(const FPaintArgs& Args, const FGeometry& AllottedGeometry, const FSlateRect& MyClippingRect, FSlateWindowElementList& OutDrawElements, int32 LayerId, const FWidgetStyle& InWidgetStyle, bool bParentEnabled) const
 {
 	ISequencer* Sequencer = GetSequencer();
@@ -89,6 +112,8 @@ int32 SCinematicTransportRange::OnPaint(const FPaintArgs& Args, const FGeometry&
 	TRange<float> WorkingRange = Sequencer->GetFocusedMovieSceneSequence()->GetMovieScene()->GetEditorData().WorkingRange;
 	TRange<float> PlaybackRange = Sequencer->GetFocusedMovieSceneSequence()->GetMovieScene()->GetPlaybackRange();
 
+	const float TimePerPixel = WorkingRange.Size<float>() / AllottedGeometry.Size.X;
+
 	FColor DarkGray(40, 40, 40);
 	FColor MidGray(80, 80, 80);
 	FColor LightGray(200, 200, 200);
@@ -120,7 +145,8 @@ int32 SCinematicTransportRange::OnPaint(const FPaintArgs& Args, const FGeometry&
 		FLinearColor(MidGray)
 	);
 
-	const float ProgressLerp = (Sequencer->GetGlobalTime() - WorkingRange.GetLowerBoundValue()) / FullRange;
+	const float CurrentTime = Sequencer->GetGlobalTime();
+	const float ProgressLerp = (CurrentTime - WorkingRange.GetLowerBoundValue()) / FullRange;
 
 	// Draw the playback progress
 	if (ProgressLerp > PlaybackStartLerp)
@@ -137,6 +163,53 @@ int32 SCinematicTransportRange::OnPaint(const FPaintArgs& Args, const FGeometry&
 		);
 	}
 
+	bool bPlayMarkerOnKey = false;
+
+	const FLinearColor KeyFrameColor = FEditorStyle::GetSlateColor("SelectionColor").GetColor(FWidgetStyle());
+
+	// Draw the current key collection tick marks
+	{
+		static const float BrushWidth = 7.f;
+		static const float BrushHeight = 7.f;
+
+		if (ActiveKeyCollection.IsValid())
+		{
+			const float BrushOffsetY = TrackOffsetY + TrackHeight * .5f - BrushHeight * .5f;
+			const FSlateBrush* KeyBrush = FLevelSequenceEditorStyle::Get()->GetBrush("LevelSequenceEditor.CinematicViewportTransportRangeKey");
+
+			ActiveKeyCollection->IterateKeys([&](float Time){
+				if (Time < WorkingRange.GetLowerBoundValue() || Time > WorkingRange.GetUpperBoundValue())
+				{
+					// continue iteration
+					return true;
+				}
+
+				if (FMath::IsNearlyEqual(CurrentTime, Time, TimePerPixel*.5f))
+				{
+					bPlayMarkerOnKey = true;
+				}
+
+				float Lerp = (Time - WorkingRange.GetLowerBoundValue()) / FullRange;
+
+				FSlateDrawElement::MakeBox(
+					OutDrawElements,
+					LayerId+2,
+					AllottedGeometry.ToPaintGeometry(
+						FVector2D(BrushWidth, BrushHeight),
+						FSlateLayoutTransform(FVector2D(AllottedGeometry.Size.X*Lerp - BrushWidth*.5f, BrushOffsetY))
+					),
+					KeyBrush,
+					MyClippingRect,
+					DrawEffects,
+					KeyFrameColor
+				);
+
+				// continue iteration
+				return true;
+			});
+		}
+	}
+
 	// Draw the play marker
 	{
 		static const float BrushWidth = 11.f, BrushHeight = 6.f;
@@ -149,23 +222,26 @@ int32 SCinematicTransportRange::OnPaint(const FPaintArgs& Args, const FGeometry&
 			FLevelSequenceEditorStyle::Get()->GetBrush("LevelSequenceEditor.CinematicViewportPlayMarker"),
 			MyClippingRect,
 			DrawEffects,
-			LightGray
+			bPlayMarkerOnKey ? KeyFrameColor : LightGray
 		);
 
-		TArray<FVector2D> LinePoints;
-		LinePoints.Add(FVector2D(PositionX, BrushHeight));
-		LinePoints.Add(FVector2D(PositionX, AllottedGeometry.Size.Y));
-
-		FSlateDrawElement::MakeLines(
-			OutDrawElements,
-			LayerId,
-			AllottedGeometry.ToPaintGeometry(),
-			LinePoints,
-			MyClippingRect,
-			DrawEffects,
-			LightGray,
-			false
-		);
+		if (!bPlayMarkerOnKey)
+		{
+			TArray<FVector2D> LinePoints;
+			LinePoints.Add(FVector2D(PositionX, BrushHeight));
+			LinePoints.Add(FVector2D(PositionX, AllottedGeometry.Size.Y));
+
+			FSlateDrawElement::MakeLines(
+				OutDrawElements,
+				LayerId,
+				AllottedGeometry.ToPaintGeometry(),
+				LinePoints,
+				MyClippingRect,
+				DrawEffects,
+				LightGray,
+				false
+			);
+		}
 	}
 
 	// Draw the play bounds
diff --git a/Engine/Plugins/MovieScene/LevelSequenceEditor/Source/LevelSequenceEditor/Private/CinematicViewport/SCinematicTransportRange.h b/Engine/Plugins/MovieScene/LevelSequenceEditor/Source/LevelSequenceEditor/Private/CinematicViewport/SCinematicTransportRange.h
--- a/Engine/Plugins/MovieScene/LevelSequenceEditor/Source/LevelSequenceEditor/Private/CinematicViewport/SCinematicTransportRange.h
+++ b/Engine/Plugins/MovieScene/LevelSequenceEditor/Source/LevelSequenceEditor/Private/CinematicViewport/SCinematicTransportRange.h
@@ -2,6 +2,8 @@
 
 #pragma once
 
+#include "ISequencerKeyCollection.h"
+
 class SCinematicTransportRange : public SCompoundWidget
 {
 public:
@@ -16,6 +18,7 @@ public:
 	void SetSequencer(TWeakPtr<ISequencer> InSequencer);
 
 	virtual FVector2D ComputeDesiredSize(float) const override;
+	virtual void Tick(const FGeometry& AllottedGeometry, const double InCurrentTime, const float InDeltaTime) override;
 	virtual int32 OnPaint(const FPaintArgs& Args, const FGeometry& AllottedGeometry, const FSlateRect& MyClippingRect, FSlateWindowElementList& OutDrawElements, int32 LayerId, const FWidgetStyle& InWidgetStyle, bool bParentEnabled) const override;
 	virtual FReply OnMouseButtonDown(const FGeometry& MyGeometry, const FPointerEvent& MouseEvent) override;
 	virtual FReply OnMouseMove(const FGeometry& MyGeometry, const FPointerEvent& MouseEvent);
@@ -31,5 +34,8 @@ private:
 	/** The sequencer that we're controlling */
 	TWeakPtr<ISequencer> WeakSequencer;
 
+	/** The collection of keys for the currently active sequencer selection */
+	TUniquePtr<ISequencerKeyCollection> ActiveKeyCollection;
+
 	bool bDraggingTime;
 };
\ No newline at end of file
diff --git a/Engine/Plugins/MovieScene/LevelSequenceEditor/Source/LevelSequenceEditor/Private/LevelSequenceEditorCommands.cpp b/Engine/Plugins/MovieScene/LevelSequenceEditor/Source/LevelSequenceEditor/Private/LevelSequenceEditorCommands.cpp
--- a/Engine/Plugins/MovieScene/LevelSequenceEditor/Source/LevelSequenceEditor/Private/LevelSequenceEditorCommands.cpp
+++ b/Engine/Plugins/MovieScene/LevelSequenceEditor/Source/LevelSequenceEditor/Private/LevelSequenceEditorCommands.cpp
@@ -13,7 +13,7 @@ FLevelSequenceEditorCommands::FLevelSequenceEditorCommands()
 void FLevelSequenceEditorCommands::RegisterCommands()
 {
 	UI_COMMAND(CreateNewLevelSequenceInLevel, "Add Level Sequence", "Create a new level sequence asset, and place an instance of it in this level", EUserInterfaceActionType::Button, FInputChord());
-	UI_COMMAND(ToggleImmersive, "Immersive Mode", "Switches this cinematic viewport between immersive mode and regular mode", EUserInterfaceActionType::ToggleButton, FInputChord( EKeys::F11 ) );
+	UI_COMMAND(ToggleCinematicViewportCommand, "Cinematic Viewport", "A viewport layout tailored to cinematic preview", EUserInterfaceActionType::RadioButton, FInputChord(EKeys::C, EModifierKey::Shift));
 }
 
 #undef LOCTEXT_NAMESPACE
\ No newline at end of file
diff --git a/Engine/Plugins/MovieScene/LevelSequenceEditor/Source/LevelSequenceEditor/Private/LevelSequenceEditorCommands.h b/Engine/Plugins/MovieScene/LevelSequenceEditor/Source/LevelSequenceEditor/Private/LevelSequenceEditorCommands.h
--- a/Engine/Plugins/MovieScene/LevelSequenceEditor/Source/LevelSequenceEditor/Private/LevelSequenceEditorCommands.h
+++ b/Engine/Plugins/MovieScene/LevelSequenceEditor/Source/LevelSequenceEditor/Private/LevelSequenceEditorCommands.h
@@ -14,5 +14,5 @@ public:
 	virtual void RegisterCommands() override;
 
 	TSharedPtr<FUICommandInfo> CreateNewLevelSequenceInLevel;
-	TSharedPtr<FUICommandInfo> ToggleImmersive;
+	TSharedPtr<FUICommandInfo> ToggleCinematicViewportCommand;
 };
diff --git a/Engine/Plugins/MovieScene/LevelSequenceEditor/Source/LevelSequenceEditor/Private/LevelSequenceEditorModule.cpp b/Engine/Plugins/MovieScene/LevelSequenceEditor/Source/LevelSequenceEditor/Private/LevelSequenceEditorModule.cpp
--- a/Engine/Plugins/MovieScene/LevelSequenceEditor/Source/LevelSequenceEditor/Private/LevelSequenceEditorModule.cpp
+++ b/Engine/Plugins/MovieScene/LevelSequenceEditor/Source/LevelSequenceEditor/Private/LevelSequenceEditorModule.cpp
@@ -2,7 +2,8 @@
 
 #include "LevelSequenceEditorPCH.h"
 #include "CineCameraActor.h"
-#include "CinematicLevelViewportLayout.h"
+#include "CameraRig_Crane.h"
+#include "CameraRig_Rail.h"
 #include "IPlacementModeModule.h"
 #include "ISettingsModule.h"
 #include "LevelEditor.h"
@@ -10,6 +11,7 @@
 #include "LevelSequenceEditorStyle.h"
 #include "ModuleInterface.h"
 #include "PropertyEditorModule.h"
+#include "CinematicViewport/CinematicViewportLayoutEntity.h"
 
 
 #define LOCTEXT_NAMESPACE "LevelSequenceEditor"
@@ -84,17 +86,9 @@ protected:
 	{
 		FLevelEditorModule& LevelEditorModule = FModuleManager::LoadModuleChecked<FLevelEditorModule>("LevelEditor");
 
-		FCustomViewportLayoutDefinition TwoPaneDefinition = FCustomViewportLayoutDefinition::FromType<FCinematicLevelViewportLayout_OnePane>();
-		TwoPaneDefinition.DisplayName = LOCTEXT("OnePaneCinematicLayoutName", "Cinematic");
-		TwoPaneDefinition.Description = LOCTEXT("OnePaneCinematicLayoutDesc", "A viewport layout tailored to cinematic preview");
-		TwoPaneDefinition.Icon = FSlateIcon("LevelSequenceEditorStyle", "LevelSequenceEditor.OnePaneCinematicViewportLayout");
-		LevelEditorModule.RegisterCustomViewportLayout("OnePaneCinematic", TwoPaneDefinition);
-
-		FCustomViewportLayoutDefinition OnePaneDefinition = FCustomViewportLayoutDefinition::FromType<FCinematicLevelViewportLayout_TwoPane>();
-		OnePaneDefinition.DisplayName = LOCTEXT("TwoPaneCinematicLayoutName", "Two Pane Cinematic");
-		OnePaneDefinition.Description = LOCTEXT("TwoPaneCinematicLayoutDesc", "A viewport layout comprising an edit viewport, and a cinematic preview viewport");
-		OnePaneDefinition.Icon = FSlateIcon("LevelSequenceEditorStyle", "LevelSequenceEditor.TwoPaneCinematicViewportLayout");
-		LevelEditorModule.RegisterCustomViewportLayout("TwoPaneCinematic", OnePaneDefinition);
+		FViewportTypeDefinition CinematicViewportType = FViewportTypeDefinition::FromType<FCinematicViewportLayoutEntity>();
+		CinematicViewportType.ToggleCommand = FLevelSequenceEditorCommands::Get().ToggleCinematicViewportCommand;
+		LevelEditorModule.RegisterViewportType("Cinematic", CinematicViewportType);
 	}
 
 	/** Register menu extensions for the level editor toolbar. */
@@ -129,6 +123,8 @@ protected:
 
 		IPlacementModeModule::Get().RegisterPlacementCategory(Info);
 		IPlacementModeModule::Get().RegisterPlaceableItem(Info.UniqueHandle, MakeShareable( new FPlaceableItem(nullptr, FAssetData(ACineCameraActor::StaticClass())) ));
+		IPlacementModeModule::Get().RegisterPlaceableItem(Info.UniqueHandle, MakeShareable( new FPlaceableItem(nullptr, FAssetData(ACameraRig_Crane::StaticClass())) ));
+		IPlacementModeModule::Get().RegisterPlaceableItem(Info.UniqueHandle, MakeShareable( new FPlaceableItem(nullptr, FAssetData(ACameraRig_Rail::StaticClass())) ));
 	}
 
 	/** Register settings objects. */
@@ -181,7 +177,7 @@ protected:
 		FLevelEditorModule* LevelEditorModule = FModuleManager::GetModulePtr<FLevelEditorModule>("LevelEditor");
 		if (LevelEditorModule)
 		{
-			LevelEditorModule->UnRegisterCustomViewportLayout("Cinematic");
+			LevelEditorModule->UnregisterViewportType("Cinematic");
 		}
 	}
 
diff --git a/Engine/Plugins/MovieScene/LevelSequenceEditor/Source/LevelSequenceEditor/Private/LevelSequenceEditorToolkit.cpp b/Engine/Plugins/MovieScene/LevelSequenceEditor/Source/LevelSequenceEditor/Private/LevelSequenceEditorToolkit.cpp
--- a/Engine/Plugins/MovieScene/LevelSequenceEditor/Source/LevelSequenceEditor/Private/LevelSequenceEditorToolkit.cpp
+++ b/Engine/Plugins/MovieScene/LevelSequenceEditor/Source/LevelSequenceEditor/Private/LevelSequenceEditorToolkit.cpp
@@ -12,6 +12,7 @@
 #include "MovieScene.h"
 #include "MovieSceneBinding.h"
 #include "MovieSceneMaterialTrack.h"
+#include "MovieScenePropertyTrack.h"
 #include "MovieSceneSequenceInstance.h"
 #include "ScopedTransaction.h"
 #include "SceneOutlinerModule.h"
@@ -22,6 +23,7 @@
 #include "MovieScene3DTransformSection.h"
 #include "MovieScene3DTransformTrack.h"
 
+#include "SequencerSpawnRegister.h"
 
 #define LOCTEXT_NAMESPACE "LevelSequenceEditor"
 
@@ -116,7 +118,7 @@ void FLevelSequenceEditorToolkit::Initialize(const EToolkitMode::Type Mode, cons
 
 	FAssetEditorToolkit::InitAssetEditor(Mode, InitToolkitHost, SequencerDefs::SequencerAppIdentifier, StandaloneDefaultLayout, bCreateDefaultStandaloneMenu, bCreateDefaultToolbar, LevelSequence);
 
-	TSharedRef<FLevelSequenceEditorSpawnRegister> SpawnRegister = MakeShareable(new FLevelSequenceEditorSpawnRegister);
+	TSharedRef<FLevelSequenceEditorSpawnRegister> SpawnRegister = MakeShareable(new TSequencerSpawnRegister<FLevelSequenceEditorSpawnRegister>);
 
 	// initialize sequencer
 	FSequencerInitParams SequencerInitParams;
@@ -148,6 +150,7 @@ void FLevelSequenceEditorToolkit::Initialize(const EToolkitMode::Type Mode, cons
 
 	Sequencer = FModuleManager::LoadModuleChecked<ISequencerModule>("Sequencer").CreateSequencer(SequencerInitParams);
 	SpawnRegister->SetSequencer(Sequencer);
+	Sequencer->OnActorAddedToSequencer().AddSP(this, &FLevelSequenceEditorToolkit::HandleActorAddedToSequencer);
 
 	if (bEditWithinLevelEditor)
 	{
@@ -286,7 +289,9 @@ void FLevelSequenceEditorToolkit::AddActorsToSequencer(AActor*const* InActors, i
 		if (!Sequencer->GetFocusedMovieSceneSequenceInstance()->FindObjectId(*ThisActor).IsValid())
 		{
 			FGuid Binding = Sequencer->CreateBinding(*ThisActor, ThisActor->GetActorLabel());
+
 			Sequencer->UpdateRuntimeInstances();
+
 			AddDefaultTracksForActor(*ThisActor, Binding);
 		}
 
@@ -297,6 +302,11 @@ void FLevelSequenceEditorToolkit::AddActorsToSequencer(AActor*const* InActors, i
 	Sequencer->NotifyMovieSceneDataChanged();
 }
 
+void FLevelSequenceEditorToolkit::HandleActorAddedToSequencer(AActor* Actor, const FGuid Binding)
+{
+	AddDefaultTracksForActor(*Actor, Binding);
+}
+
 void FLevelSequenceEditorToolkit::AddDefaultTracksForActor(AActor& Actor, const FGuid Binding)
 {
 	// get focused movie scene
@@ -332,30 +342,45 @@ void FLevelSequenceEditorToolkit::AddDefaultTracksForActor(AActor& Actor, const
 			if (TrackClass != nullptr)
 			{
 				UMovieSceneTrack* NewTrack = MovieScene->AddTrack(TrackClass, Binding);
-				UMovieSceneSection* NewSection = NewTrack->CreateNewSection();
-
-				NewTrack->AddSection(*NewSection);
 
-				// @todo sequencer: hack: setting defaults for transform tracks
-				if (NewTrack->IsA(UMovieScene3DTransformTrack::StaticClass()))
+				// Create a section for any property tracks
+				if (Cast<UMovieScenePropertyTrack>(NewTrack))
 				{
-					auto TransformSection = Cast<UMovieScene3DTransformSection>(NewSection);
+					UMovieSceneSection* NewSection = NewTrack->CreateNewSection();
+					NewTrack->AddSection(*NewSection);
 
-					const FVector Location = Actor.GetActorLocation();
-					const FRotator Rotation = Actor.GetActorRotation();
-					const FVector Scale = Actor.GetActorScale();
+					// @todo sequencer: hack: setting defaults for transform tracks
+					if (NewTrack->IsA(UMovieScene3DTransformTrack::StaticClass()))
+					{
+						auto TransformSection = Cast<UMovieScene3DTransformSection>(NewSection);
+
+						FVector Location = Actor.GetActorLocation();
+						FRotator Rotation = Actor.GetActorRotation();
+						FVector Scale = Actor.GetActorScale();
+
+						if (Actor.GetRootComponent())
+						{
+							FTransform ActorRelativeTransform = Actor.GetRootComponent()->GetRelativeTransform();
 
-					TransformSection->SetDefault(FTransformKey(EKey3DTransformChannel::Translation, EAxis::X, Location.X, false /*bUnwindRotation*/));
-					TransformSection->SetDefault(FTransformKey(EKey3DTransformChannel::Translation, EAxis::Y, Location.Y, false /*bUnwindRotation*/));
-					TransformSection->SetDefault(FTransformKey(EKey3DTransformChannel::Translation, EAxis::Z, Location.Z, false /*bUnwindRotation*/));
+							Location = ActorRelativeTransform.GetTranslation();
+							Rotation = ActorRelativeTransform.GetRotation().Rotator();
+							Scale = ActorRelativeTransform.GetScale3D();
+						}
 
-					TransformSection->SetDefault(FTransformKey(EKey3DTransformChannel::Rotation, EAxis::X, Rotation.Euler().X, false /*bUnwindRotation*/));
-					TransformSection->SetDefault(FTransformKey(EKey3DTransformChannel::Rotation, EAxis::Y, Rotation.Euler().Y, false /*bUnwindRotation*/));
-					TransformSection->SetDefault(FTransformKey(EKey3DTransformChannel::Rotation, EAxis::Z, Rotation.Euler().Z, false /*bUnwindRotation*/));
+						TransformSection->SetDefault(FTransformKey(EKey3DTransformChannel::Translation, EAxis::X, Location.X, false /*bUnwindRotation*/));
+						TransformSection->SetDefault(FTransformKey(EKey3DTransformChannel::Translation, EAxis::Y, Location.Y, false /*bUnwindRotation*/));
+						TransformSection->SetDefault(FTransformKey(EKey3DTransformChannel::Translation, EAxis::Z, Location.Z, false /*bUnwindRotation*/));
+
+						TransformSection->SetDefault(FTransformKey(EKey3DTransformChannel::Rotation, EAxis::X, Rotation.Euler().X, false /*bUnwindRotation*/));
+						TransformSection->SetDefault(FTransformKey(EKey3DTransformChannel::Rotation, EAxis::Y, Rotation.Euler().Y, false /*bUnwindRotation*/));
+						TransformSection->SetDefault(FTransformKey(EKey3DTransformChannel::Rotation, EAxis::Z, Rotation.Euler().Z, false /*bUnwindRotation*/));
+
+						TransformSection->SetDefault(FTransformKey(EKey3DTransformChannel::Scale, EAxis::X, Scale.X, false /*bUnwindRotation*/));
+						TransformSection->SetDefault(FTransformKey(EKey3DTransformChannel::Scale, EAxis::Y, Scale.Y, false /*bUnwindRotation*/));
+						TransformSection->SetDefault(FTransformKey(EKey3DTransformChannel::Scale, EAxis::Z, Scale.Z, false /*bUnwindRotation*/));
+					}
 
-					TransformSection->SetDefault(FTransformKey(EKey3DTransformChannel::Scale, EAxis::X, Scale.X, false /*bUnwindRotation*/));
-					TransformSection->SetDefault(FTransformKey(EKey3DTransformChannel::Scale, EAxis::Y, Scale.Y, false /*bUnwindRotation*/));
-					TransformSection->SetDefault(FTransformKey(EKey3DTransformChannel::Scale, EAxis::Z, Scale.Z, false /*bUnwindRotation*/));
+					NewSection->SetIsInfinite(GetSequencer()->GetInfiniteKeyAreas());
 				}
 
 				Sequencer->UpdateRuntimeInstances();
@@ -426,9 +451,9 @@ void FLevelSequenceEditorToolkit::AddDefaultTracksForActor(AActor& Actor, const
 			{
 				KeyPropertyParams.KeyParams.bCreateTrackIfMissing = true;
 				KeyPropertyParams.KeyParams.bCreateHandleIfMissing = true;
-				KeyPropertyParams.KeyParams.bCreateKeyIfUnchanged = true;
-				KeyPropertyParams.KeyParams.bCreateKeyIfEmpty = true;
-				KeyPropertyParams.KeyParams.bCreateKeyOnlyWhenAutoKeying = true;
+				KeyPropertyParams.KeyParams.bCreateKeyIfUnchanged = false;
+				KeyPropertyParams.KeyParams.bCreateKeyIfEmpty = false;
+				KeyPropertyParams.KeyParams.bCreateKeyOnlyWhenAutoKeying = false;
 			}
 
 			Sequencer->KeyProperty(KeyPropertyParams);
diff --git a/Engine/Plugins/MovieScene/LevelSequenceEditor/Source/LevelSequenceEditor/Private/LevelSequenceEditorToolkit.h b/Engine/Plugins/MovieScene/LevelSequenceEditor/Source/LevelSequenceEditor/Private/LevelSequenceEditorToolkit.h
--- a/Engine/Plugins/MovieScene/LevelSequenceEditor/Source/LevelSequenceEditor/Private/LevelSequenceEditorToolkit.h
+++ b/Engine/Plugins/MovieScene/LevelSequenceEditor/Source/LevelSequenceEditor/Private/LevelSequenceEditorToolkit.h
@@ -106,6 +106,9 @@ private:
 	/** Add the specified actors to the sequencer */
 	void AddActorsToSequencer(AActor*const* InActors, int32 NumActors);
 
+	/** Callback for actor added to sequencer. */
+	void HandleActorAddedToSequencer(AActor* Actor, const FGuid Binding);
+
 	/** Add default movie scene tracks for the given actor. */
 	void AddDefaultTracksForActor(AActor& Actor, const FGuid Binding);
 
diff --git a/Engine/Plugins/MovieScene/LevelSequenceEditor/Source/LevelSequenceEditor/Private/Styles/LevelSequenceEditorStyle.h b/Engine/Plugins/MovieScene/LevelSequenceEditor/Source/LevelSequenceEditor/Private/Styles/LevelSequenceEditorStyle.h
--- a/Engine/Plugins/MovieScene/LevelSequenceEditor/Source/LevelSequenceEditor/Private/Styles/LevelSequenceEditorStyle.h
+++ b/Engine/Plugins/MovieScene/LevelSequenceEditor/Source/LevelSequenceEditor/Private/Styles/LevelSequenceEditorStyle.h
@@ -30,24 +30,16 @@ public:
 
 		// tab icons
 		Set("LevelSequenceEditor.Tabs.Sequencer", new IMAGE_BRUSH("icon_tab_sequencer_16x", Icon16x16));
-		Set("LevelSequenceEditor.PossessNewActor", new IMAGE_BRUSH("PossessNewActor_16x", Icon16x16));
-		Set("LevelSequenceEditor.PossessNewActor.Small", new IMAGE_BRUSH("PossessNewActor_16x", Icon16x16));
+		Set("LevelSequenceEditor.PossessNewActor", new IMAGE_BRUSH("Icon_Actor_To_Sequencer_16x", Icon16x16));
+		Set("LevelSequenceEditor.PossessNewActor.Small", new IMAGE_BRUSH("Icon_Actor_To_Sequencer_16x", Icon16x16));
 		Set("LevelSequenceEditor.CreateNewLevelSequenceInLevel", new IMAGE_BRUSH("CreateNewLevelSequenceInLevel_16x", Icon16x16));
 		Set("LevelSequenceEditor.CreateNewLevelSequenceInLevel.Small", new IMAGE_BRUSH("CreateNewLevelSequenceInLevel_16x", Icon16x16));
-		Set("LevelSequenceEditor.OnePaneCinematicViewportLayout", new IMAGE_BRUSH("ViewportLayout_OnePaneCinematic", FVector2D(47, 37)));
-		Set("LevelSequenceEditor.TwoPaneCinematicViewportLayout", new IMAGE_BRUSH("ViewportLayout_TwoPaneCinematic", FVector2D(47, 37)));
-		Set("LevelSequenceEditor.OnePaneCinematicViewportLayout.Small", new IMAGE_BRUSH("ViewportLayout_OnePaneCinematic", FVector2D(47, 37)));
-		Set("LevelSequenceEditor.TwoPaneCinematicViewportLayout.Small", new IMAGE_BRUSH("ViewportLayout_TwoPaneCinematic", FVector2D(47, 37)));
-
-		Set("LevelSequenceEditor.CinematicViewportCamera", new IMAGE_BRUSH("CinematicViewportCamera", FVector2D(36, 24)));
-		Set("LevelSequenceEditor.CinematicViewportClose", new IMAGE_BRUSH("CinematicViewportClose", FVector2D(16, 16)));
-
+		
 		Set("LevelSequenceEditor.CinematicViewportPlayMarker", new IMAGE_BRUSH("CinematicViewportPlayMarker", FVector2D(11, 6)));
 		Set("LevelSequenceEditor.CinematicViewportRangeStart", new BORDER_BRUSH("CinematicViewportRangeStart", FMargin(1.f,.3f,0.f,.6f)));
 		Set("LevelSequenceEditor.CinematicViewportRangeEnd", new BORDER_BRUSH("CinematicViewportRangeEnd", FMargin(0.f,.3f,1.f,.6f)));
 
-		Set("LevelSequenceEditor.CinematicViewportSetPlayStart", new IMAGE_BRUSH("CinematicViewportSetPlayStart", FVector2D(24, 24)));
-		Set("LevelSequenceEditor.CinematicViewportSetPlayEnd", new IMAGE_BRUSH("CinematicViewportSetPlayEnd", FVector2D(24, 24)));
+		Set("LevelSequenceEditor.CinematicViewportTransportRangeKey", new IMAGE_BRUSH("CinematicViewportTransportRangeKey", FVector2D(7.f, 7.f)));
 
 		Set("FilmOverlay.DefaultThumbnail", new IMAGE_BRUSH("DefaultFilmOverlayThumbnail", FVector2D(36, 24)));
 
diff --git a/Engine/Shaders/ShaderVersion.usf b/Engine/Shaders/ShaderVersion.usf
--- a/Engine/Shaders/ShaderVersion.usf
+++ b/Engine/Shaders/ShaderVersion.usf
@@ -2,4 +2,4 @@
 // Each time the console command is executed it generates a new GUID. As this file is included
 // in common.usf (which should be included in any shader) it allows to invalidate the shader DDC.
 // 
-// GUID = 5F942FD2FE1B486CBEC6EDD6EC78A38F  
+// GUID = 4262103F821247F4868CDF3249343BBA  
diff --git a/Engine/Source/Editor/DetailCustomizations/Private/DetailCustomizations.cpp b/Engine/Source/Editor/DetailCustomizations/Private/DetailCustomizations.cpp
--- a/Engine/Source/Editor/DetailCustomizations/Private/DetailCustomizations.cpp
+++ b/Engine/Source/Editor/DetailCustomizations/Private/DetailCustomizations.cpp
@@ -109,6 +109,7 @@
 #include "ImportantToggleSettingCustomization.h"
 #include "CameraFilmbackSettingsCustomization.h"
 #include "CameraLensSettingsCustomization.h"
+#include "CameraFocusSettingsCustomization.h"
 #include "ColorStructCustomization.h"
 #include "RotatorStructCustomization.h"
 #include "VectorStructCustomization.h"
@@ -225,6 +226,7 @@ void FDetailCustomizationsModule::RegisterPropertyTypeCustomizations()
 	RegisterCustomPropertyTypeLayout("MaterialProxySettings", FOnGetPropertyTypeCustomizationInstance::CreateStatic(&FMaterialProxySettingsCustomizations::MakeInstance));
 	RegisterCustomPropertyTypeLayout("CameraFilmbackSettings", FOnGetPropertyTypeCustomizationInstance::CreateStatic(&FCameraFilmbackSettingsCustomization::MakeInstance));
 	RegisterCustomPropertyTypeLayout("CameraLensSettings", FOnGetPropertyTypeCustomizationInstance::CreateStatic(&FCameraLensSettingsCustomization::MakeInstance));
+	RegisterCustomPropertyTypeLayout("CameraFocusSettings", FOnGetPropertyTypeCustomizationInstance::CreateStatic(&FCameraFocusSettingsCustomization::MakeInstance));
 }
 
 
@@ -339,4 +341,4 @@ void FDetailCustomizationsModule::RegisterCustomPropertyTypeLayout(FName Propert
 	static FName PropertyEditor("PropertyEditor");
 	FPropertyEditorModule& PropertyModule = FModuleManager::GetModuleChecked<FPropertyEditorModule>(PropertyEditor);
 	PropertyModule.RegisterCustomPropertyTypeLayout(PropertyTypeName, PropertyTypeLayoutDelegate);
-}
+}
\ No newline at end of file
diff --git a/Engine/Source/Editor/DetailCustomizations/Private/DirectoryPathStructCustomization.cpp b/Engine/Source/Editor/DetailCustomizations/Private/DirectoryPathStructCustomization.cpp
--- a/Engine/Source/Editor/DetailCustomizations/Private/DirectoryPathStructCustomization.cpp
+++ b/Engine/Source/Editor/DetailCustomizations/Private/DirectoryPathStructCustomization.cpp
@@ -4,6 +4,7 @@
 #include "DirectoryPathStructCustomization.h"
 #include "DesktopPlatformModule.h"
 #include "ContentBrowserDelegates.h"
+#include "ContentBrowserModule.h"
 
 #define LOCTEXT_NAMESPACE "DirectoryPathStructCustomization"
 
@@ -18,11 +19,54 @@ void FDirectoryPathStructCustomization::CustomizeHeader( TSharedRef<IPropertyHan
 
 	const bool bRelativeToGameContentDir = StructPropertyHandle->HasMetaData( TEXT("RelativeToGameContentDir") );
 	const bool bUseRelativePath = StructPropertyHandle->HasMetaData( TEXT("RelativePath") );
-	
+	const bool bContentDir = StructPropertyHandle->HasMetaData( TEXT("ContentDir") );
+
 	AbsoluteGameContentDir = FPaths::ConvertRelativePathToFull(FPaths::GameContentDir());
 
 	if(PathProperty.IsValid())
 	{
+		TSharedPtr<SWidget> PickerWidget = nullptr;
+
+		if(bContentDir)
+		{
+			FPathPickerConfig PathPickerConfig;
+			PathPickerConfig.bAllowContextMenu = false;
+			PathPickerConfig.OnPathSelected = FOnPathSelected::CreateSP(this, &FDirectoryPathStructCustomization::OnPathPicked, PathProperty.ToSharedRef());
+
+			FContentBrowserModule& ContentBrowserModule = FModuleManager::LoadModuleChecked<FContentBrowserModule>("ContentBrowser");
+
+			PickerWidget = SAssignNew(PickerButton, SComboButton)
+			.ButtonStyle( FEditorStyle::Get(), "HoverHintOnly" )
+			.ToolTipText( LOCTEXT( "FolderComboToolTipText", "Choose a content directory") )
+			.ContentPadding( 2.0f )
+			.ForegroundColor( FSlateColor::UseForeground() )
+			.IsFocusable( false )
+			.MenuContent()
+			[
+				SNew(SBox)
+				.WidthOverride(300.0f)
+				.HeightOverride(300.0f)
+				[
+					ContentBrowserModule.Get().CreatePathPicker(PathPickerConfig)
+				]
+			];			
+		}
+		else
+		{
+			PickerWidget = SAssignNew(BrowseButton, SButton)
+			.ButtonStyle( FEditorStyle::Get(), "HoverHintOnly" )
+			.ToolTipText( LOCTEXT( "FolderButtonToolTipText", "Choose a directory from this computer") )
+			.OnClicked( FOnClicked::CreateSP(this, &FDirectoryPathStructCustomization::OnPickDirectory, PathProperty.ToSharedRef(), bRelativeToGameContentDir, bUseRelativePath) )
+			.ContentPadding( 2.0f )
+			.ForegroundColor( FSlateColor::UseForeground() )
+			.IsFocusable( false )
+			[
+				SNew( SImage )
+				.Image( FEditorStyle::GetBrush("PropertyWindow.Button_Ellipsis") )
+				.ColorAndOpacity( FSlateColor::UseForeground() )
+			];
+		}
+
 		HeaderRow.ValueContent()
 		.MinDesiredWidth(125.0f)
 		.MaxDesiredWidth(600.0f)
@@ -39,18 +83,7 @@ void FDirectoryPathStructCustomization::CustomizeHeader( TSharedRef<IPropertyHan
 			.Padding(FMargin(4.0f, 0.0f, 0.0f, 0.0f))
 			.VAlign(VAlign_Center)
 			[
-				SAssignNew(BrowseButton, SButton)
-				.ButtonStyle( FEditorStyle::Get(), "HoverHintOnly" )
-				.ToolTipText( LOCTEXT( "FolderButtonToolTipText", "Choose a directory from this computer") )
-				.OnClicked( FOnClicked::CreateSP(this, &FDirectoryPathStructCustomization::OnPickDirectory, PathProperty.ToSharedRef(), bRelativeToGameContentDir, bUseRelativePath) )
-				.ContentPadding( 2.0f )
-				.ForegroundColor( FSlateColor::UseForeground() )
-				.IsFocusable( false )
-				[
-					SNew( SImage )
-					.Image( FEditorStyle::GetBrush("PropertyWindow.Button_Ellipsis") )
-					.ColorAndOpacity( FSlateColor::UseForeground() )
-				]
+				PickerWidget.ToSharedRef()
 			]
 		]
 		.NameContent()
@@ -132,4 +165,11 @@ bool FDirectoryPathStructCustomization::IsValidPath(const FString& AbsolutePath,
 	return true;
 }
 
+void FDirectoryPathStructCustomization::OnPathPicked(const FString& Path, TSharedRef<IPropertyHandle> PropertyHandle)
+{
+	PickerButton->SetIsOpen(false);
+
+	PropertyHandle->SetValue(Path);
+}
+
 #undef LOCTEXT_NAMESPACE
\ No newline at end of file
diff --git a/Engine/Source/Editor/DetailCustomizations/Private/DirectoryPathStructCustomization.h b/Engine/Source/Editor/DetailCustomizations/Private/DirectoryPathStructCustomization.h
--- a/Engine/Source/Editor/DetailCustomizations/Private/DirectoryPathStructCustomization.h
+++ b/Engine/Source/Editor/DetailCustomizations/Private/DirectoryPathStructCustomization.h
@@ -23,9 +23,15 @@ private:
 	/** Check whether that the chosen path is valid */
 	bool IsValidPath(const FString& AbsolutePath, const bool bRelativeToGameContentDir, FText* const OutReason = nullptr) const;
 
+	/** Called when a path is picked from the path picker */
+	void OnPathPicked(const FString& Path, TSharedRef<IPropertyHandle> PropertyHandle);
+
 	/** The browse button widget */
 	TSharedPtr<SButton> BrowseButton;
 
+	/** The pick button widget */
+	TSharedPtr<SComboButton> PickerButton;
+	
 	/** Absolute path to the game content directory */
 	FString AbsoluteGameContentDir;
 };
\ No newline at end of file
