Fixed ViewportInteraction regression with assertion when opening Sequencer

diff --git a/Engine/Plugins/MovieScene/LevelSequenceEditor/Source/LevelSequenceEditor/Private/LevelSequenceEditorToolkit.cpp b/Engine/Plugins/MovieScene/LevelSequenceEditor/Source/LevelSequenceEditor/Private/LevelSequenceEditorToolkit.cpp
--- a/Engine/Plugins/MovieScene/LevelSequenceEditor/Source/LevelSequenceEditor/Private/LevelSequenceEditorToolkit.cpp
+++ b/Engine/Plugins/MovieScene/LevelSequenceEditor/Source/LevelSequenceEditor/Private/LevelSequenceEditorToolkit.cpp
@@ -280,7 +280,7 @@ void FLevelSequenceEditorToolkit::Initialize(const EToolkitMode::Type Mode, cons
 
 	{
 		const UWorld* World = CastChecked<UWorld>( GetLevelSequenceEditorPlaybackContext() );
-		UVREditorMode* VRMode = CastChecked<UVREditorMode>( GEditor->GetEditorWorldExtensionsManager()->GetEditorWorldExtensions( World )->FindExtension( UVREditorMode::StaticClass() ) );
+		UVREditorMode* VRMode = Cast<UVREditorMode>( GEditor->GetEditorWorldExtensionsManager()->GetEditorWorldExtensions( World )->FindExtension( UVREditorMode::StaticClass() ) );
 		if (VRMode != nullptr)
 		{
 			VRMode->OnVREditingModeExit_Handler.BindSP(this, &FLevelSequenceEditorToolkit::HandleVREditorModeExit);
