Copying //UE4/Dev-Sequencer to //UE4/Dev-Main (Source: //UE4/Dev-Sequencer @ 4044339)

diff --git a/Engine/Plugins/Media/MediaCompositing/Source/MediaCompositing/Private/MovieScene/MovieSceneMediaTemplate.cpp b/Engine/Plugins/Media/MediaCompositing/Source/MediaCompositing/Private/MovieScene/MovieSceneMediaTemplate.cpp
--- a/Engine/Plugins/Media/MediaCompositing/Source/MediaCompositing/Private/MovieScene/MovieSceneMediaTemplate.cpp
+++ b/Engine/Plugins/Media/MediaCompositing/Source/MediaCompositing/Private/MovieScene/MovieSceneMediaTemplate.cpp
@@ -4,6 +4,7 @@
 
 #include "Math/UnrealMathUtility.h"
 #include "MediaPlayer.h"
+#include "MediaPlayerFacade.h"
 #include "MediaSoundComponent.h"
 #include "MediaSource.h"
 #include "MediaTexture.h"
@@ -15,6 +16,9 @@
 #include "MovieSceneMediaTrack.h"
 
 
+#define MOVIESCENEMEDIATEMPLATE_TRACE_EVALUATION 0
+
+
 /* Local helpers
  *****************************************************************************/
 
@@ -81,6 +85,7 @@ struct FMediaSectionExecutionToken
 			return;
 		}
 
+		// seek on open if necessary
 		if (MediaPlayer->IsPreparing())
 		{
 			SectionData.SeekOnOpen(CurrentTime);
@@ -88,8 +93,19 @@ struct FMediaSectionExecutionToken
 			return;
 		}
 
+		const FTimespan MediaDuration = MediaPlayer->GetDuration();
+
+		if (MediaDuration.IsZero())
+		{
+			return; // media has no length
+		}
+
 		// update media player
-		const FTimespan MediaTime = CurrentTime % MediaPlayer->GetDuration();
+		const FTimespan MediaTime = CurrentTime % MediaDuration;
+
+		#if MOVIESCENEMEDIATEMPLATE_TRACE_EVALUATION
+			GLog->Logf(ELogVerbosity::Log, TEXT("Executing time %s, MediaTime %s"), *CurrentTime.ToString(TEXT("%h:%m:%s.%t")), *MediaTime.ToString(TEXT("%h:%m:%s.%t")));
+		#endif
 
 		if (Context.GetStatus() == EMovieScenePlayerStatus::Playing)
 		{
@@ -102,6 +118,8 @@ struct FMediaSectionExecutionToken
 			{
 				MediaPlayer->Seek(MediaTime);
 			}
+
+			MediaPlayer->SetBlockOnTime(MediaPlayer->GetTime());
 		}
 		else
 		{
@@ -111,6 +129,7 @@ struct FMediaSectionExecutionToken
 			}
 
 			MediaPlayer->Seek(MediaTime);
+			MediaPlayer->SetBlockOnTime(FTimespan::MinValue());
 		}
 	}
 
@@ -131,11 +150,11 @@ FMovieSceneMediaSectionTemplate::FMovieSceneMediaSectionTemplate(const UMovieSce
 	Params.MediaSource = InSection.GetMediaSource();
 	Params.MediaTexture = InSection.MediaTexture;
 
-	if (InSection.GetRange().GetLowerBound().IsClosed())
+	if (InSection.HasStartFrame())
 	{
 		Params.SectionStartFrame = InSection.GetRange().GetLowerBoundValue();
 	}
-	if (InSection.GetRange().GetUpperBound().IsClosed())
+	if (InSection.HasEndFrame())
 	{
 		Params.SectionEndFrame = InSection.GetRange().GetUpperBoundValue();
 	}
@@ -166,12 +185,23 @@ void FMovieSceneMediaSectionTemplate::Evaluate(const FMovieSceneEvaluationOperan
 	else if (!Context.IsPostRoll() && (Context.GetTime().FrameNumber < Params.SectionEndFrame))
 	{
 		const FFrameRate FrameRate = Context.GetFrameRate();
-		const FFrameTime CurrentTime(Context.GetTime().FrameNumber - Params.SectionStartFrame);
+		const FFrameTime FrameTime(Context.GetTime().FrameNumber - Params.SectionStartFrame);
 		const int64 DenominatorTicks = FrameRate.Denominator * ETimespan::TicksPerSecond;
-		const int64 CurrentTicks = FMath::DivideAndRoundNearest(int64(CurrentTime.GetFrame().Value * DenominatorTicks), int64(FrameRate.Numerator));
-		const int64 CurrentSubTicks = FMath::DivideAndRoundNearest(int64(CurrentTime.GetSubFrame() * DenominatorTicks), int64(FrameRate.Numerator));
-
-		ExecutionTokens.Add(FMediaSectionExecutionToken(Params.MediaSource, FTimespan(CurrentTicks + CurrentSubTicks)));
+		const int64 FrameTicks = FMath::DivideAndRoundNearest(int64(FrameTime.GetFrame().Value * DenominatorTicks), int64(FrameRate.Numerator));
+		const int64 FrameSubTicks = FMath::DivideAndRoundNearest(int64(FrameTime.GetSubFrame() * DenominatorTicks), int64(FrameRate.Numerator));
+
+		#if MOVIESCENEMEDIATEMPLATE_TRACE_EVALUATION
+			GLog->Logf(ELogVerbosity::Log, TEXT("Evaluating frame %i+%f, FrameRate %i/%i, FrameTicks %d+%d"),
+				Context.GetTime().GetFrame().Value,
+				Context.GetTime().GetSubFrame(),
+				FrameRate.Numerator,
+				FrameRate.Denominator,
+				FrameTicks,
+				FrameSubTicks
+			);
+		#endif
+
+		ExecutionTokens.Add(FMediaSectionExecutionToken(Params.MediaSource, FTimespan(FrameTicks + FrameSubTicks)));
 	}
 }
 
@@ -204,10 +234,18 @@ void FMovieSceneMediaSectionTemplate::Initialize(const FMovieSceneEvaluationOper
 	{
 		if (IsEvaluating)
 		{
+			#if MOVIESCENEMEDIATEMPLATE_TRACE_EVALUATION
+				GLog->Logf(ELogVerbosity::Log, TEXT("Setting media player %p on media sound component %p"), MediaPlayer, Params.MediaSoundComponent);
+			#endif
+
 			Params.MediaSoundComponent->SetMediaPlayer(MediaPlayer);
 		}
 		else if (Params.MediaSoundComponent->GetMediaPlayer() == MediaPlayer)
 		{
+			#if MOVIESCENEMEDIATEMPLATE_TRACE_EVALUATION
+				GLog->Logf(ELogVerbosity::Log, TEXT("Resetting media player on media sound component %p"), Params.MediaSoundComponent);
+			#endif
+
 			Params.MediaSoundComponent->SetMediaPlayer(nullptr);
 		}
 	}
@@ -216,10 +254,18 @@ void FMovieSceneMediaSectionTemplate::Initialize(const FMovieSceneEvaluationOper
 	{
 		if (IsEvaluating)
 		{
+			#if MOVIESCENEMEDIATEMPLATE_TRACE_EVALUATION
+				GLog->Logf(ELogVerbosity::Log, TEXT("Setting media player %p on media texture %p"), MediaPlayer, Params.MediaTexture);
+			#endif
+
 			Params.MediaTexture->SetMediaPlayer(MediaPlayer);
 		}
 		else if (Params.MediaTexture->GetMediaPlayer() == MediaPlayer)
 		{
+			#if MOVIESCENEMEDIATEMPLATE_TRACE_EVALUATION
+				GLog->Logf(ELogVerbosity::Log, TEXT("Resetting media player on media texture %p"), Params.MediaTexture);
+			#endif
+
 			Params.MediaTexture->SetMediaPlayer(nullptr);
 		}
 	}
diff --git a/Engine/Plugins/Media/MediaCompositing/Source/MediaCompositing/Private/MovieScene/MovieSceneMediaTrack.cpp b/Engine/Plugins/Media/MediaCompositing/Source/MediaCompositing/Private/MovieScene/MovieSceneMediaTrack.cpp
--- a/Engine/Plugins/Media/MediaCompositing/Source/MediaCompositing/Private/MovieScene/MovieSceneMediaTrack.cpp
+++ b/Engine/Plugins/Media/MediaCompositing/Source/MediaCompositing/Private/MovieScene/MovieSceneMediaTrack.cpp
@@ -37,9 +37,8 @@ UMovieSceneMediaTrack::UMovieSceneMediaTrack(const FObjectInitializer& ObjectIni
 UMovieSceneSection* UMovieSceneMediaTrack::AddNewMediaSourceOnRow(UMediaSource& MediaSource, FFrameNumber Time, int32 RowIndex)
 {
 	const float DefaultMediaSectionDuration = 1.0f;
-	FFrameRate FrameRate = GetTypedOuter<UMovieScene>()->GetFrameResolution();
-
-	FFrameTime DurationToUse = DefaultMediaSectionDuration * FrameRate;
+	FFrameRate TickResolution = GetTypedOuter<UMovieScene>()->GetTickResolution();
+	FFrameTime DurationToUse  = DefaultMediaSectionDuration * TickResolution;
 
 	// add the section
 	UMovieSceneMediaSection* NewSection = NewObject<UMovieSceneMediaSection>(this);
diff --git a/Engine/Plugins/Media/MediaCompositing/Source/MediaCompositingEditor/MediaCompositingEditor.Build.cs b/Engine/Plugins/Media/MediaCompositing/Source/MediaCompositingEditor/MediaCompositingEditor.Build.cs
--- a/Engine/Plugins/Media/MediaCompositing/Source/MediaCompositingEditor/MediaCompositingEditor.Build.cs
+++ b/Engine/Plugins/Media/MediaCompositing/Source/MediaCompositingEditor/MediaCompositingEditor.Build.cs
@@ -24,6 +24,7 @@ namespace UnrealBuildTool.Rules
 					"Slate",
 					"SlateCore",
 					"UnrealEd",
+					"TimeManagement"
 				});
 
 			PrivateIncludePaths.AddRange(
diff --git a/Engine/Plugins/Media/MediaCompositing/Source/MediaCompositingEditor/Private/Sequencer/MediaThumbnailSection.cpp b/Engine/Plugins/Media/MediaCompositing/Source/MediaCompositingEditor/Private/Sequencer/MediaThumbnailSection.cpp
--- a/Engine/Plugins/Media/MediaCompositing/Source/MediaCompositingEditor/Private/Sequencer/MediaThumbnailSection.cpp
+++ b/Engine/Plugins/Media/MediaCompositing/Source/MediaCompositingEditor/Private/Sequencer/MediaThumbnailSection.cpp
@@ -171,7 +171,7 @@ void FMediaThumbnailSection::SetSingleTime(double GlobalTime)
 
 	if (MediaSection != nullptr)
 	{
-		double StartTime = MediaSection->GetInclusiveStartFrame() / MediaSection->GetTypedOuter<UMovieScene>()->GetFrameResolution();
+		double StartTime = MediaSection->GetInclusiveStartFrame() / MediaSection->GetTypedOuter<UMovieScene>()->GetTickResolution();
 		MediaSection->SetThumbnailReferenceOffset(GlobalTime - StartTime);
 	}
 }
@@ -310,8 +310,8 @@ void FMediaThumbnailSection::DrawLoopIndicators(FSequencerSectionPainter& InPain
 {
 	static const FSlateBrush* GenericBrush = FCoreStyle::Get().GetBrush("GenericWhiteBox");
 
-	FFrameRate FrameResolution = Section->GetTypedOuter<UMovieScene>()->GetFrameResolution();
-	double SectionDuration = FFrameTime(MovieScene::DiscreteSize(Section->GetRange())) / FrameResolution;
+	FFrameRate TickResolution = Section->GetTypedOuter<UMovieScene>()->GetTickResolution();
+	double SectionDuration = FFrameTime(MovieScene::DiscreteSize(Section->GetRange())) / TickResolution;
 	const float MediaSizeX = MediaDuration.GetTotalSeconds() * SectionSize.X / SectionDuration;
 	float DrawOffset = 0.0f;
 
@@ -335,8 +335,8 @@ void FMediaThumbnailSection::DrawSampleStates(FSequencerSectionPainter& InPainte
 {
 	static const FSlateBrush* GenericBrush = FCoreStyle::Get().GetBrush("GenericWhiteBox");
 
-	FFrameRate FrameResolution = Section->GetTypedOuter<UMovieScene>()->GetFrameResolution();
-	double SectionDuration = FFrameTime(MovieScene::DiscreteSize(Section->GetRange())) / FrameResolution;
+	FFrameRate TickResolution = Section->GetTypedOuter<UMovieScene>()->GetTickResolution();
+	double SectionDuration = FFrameTime(MovieScene::DiscreteSize(Section->GetRange())) / TickResolution;
 	const float MediaSizeX = MediaDuration.GetTotalSeconds() * SectionSize.X / SectionDuration;
 
 	TArray<TRange<FTimespan>> Ranges;
diff --git a/Engine/Plugins/Media/MediaCompositing/Source/MediaCompositingEditor/Private/Sequencer/MediaTrackEditor.cpp b/Engine/Plugins/Media/MediaCompositing/Source/MediaCompositingEditor/Private/Sequencer/MediaTrackEditor.cpp
--- a/Engine/Plugins/Media/MediaCompositing/Source/MediaCompositingEditor/Private/Sequencer/MediaTrackEditor.cpp
+++ b/Engine/Plugins/Media/MediaCompositing/Source/MediaCompositingEditor/Private/Sequencer/MediaTrackEditor.cpp
@@ -15,8 +15,6 @@
 #include "MovieSceneMediaSection.h"
 #include "MovieSceneMediaTrack.h"
 #include "MovieSceneToolsUserSettings.h"
-#include "MediaPlane.h"
-#include "MediaPlaneComponent.h"
 #include "SequencerUtilities.h"
 #include "TrackEditorThumbnail/TrackEditorThumbnailPool.h"
 #include "Widgets/SBoxPanel.h"
@@ -92,6 +90,7 @@ TSharedPtr<SWidget> FMediaTrackEditor::BuildOutlinerEditWidget(const FGuid& Obje
 		FAssetPickerConfig AssetPickerConfig;
 		{
 			AssetPickerConfig.OnAssetSelected = FOnAssetSelected::CreateRaw(this, &FMediaTrackEditor::AddNewSection, MediaTrack);
+			AssetPickerConfig.OnAssetEnterPressed = FOnAssetEnterPressed::CreateRaw(this, &FMediaTrackEditor::AddNewSectionEnterPressed, MediaTrack);
 			AssetPickerConfig.bAllowNullSelection = false;
 			AssetPickerConfig.InitialAssetViewType = EAssetViewType::List;
 			AssetPickerConfig.Filter.bRecursiveClasses = true;
@@ -159,6 +158,12 @@ TSharedRef<ISequencerSection> FMediaTrackEditor::MakeSectionInterface(UMovieScen
 }
 
 
+bool FMediaTrackEditor::SupportsSequence(UMovieSceneSequence* InSequence) const
+{
+	return (InSequence != nullptr) && (InSequence->GetClass()->GetName() == TEXT("LevelSequence"));
+}
+
+
 bool FMediaTrackEditor::SupportsType(TSubclassOf<UMovieSceneTrack> TrackClass) const
 {
 	return TrackClass.Get() && TrackClass.Get()->IsChildOf(UMovieSceneMediaTrack::StaticClass());
@@ -265,6 +270,13 @@ void FMediaTrackEditor::AddNewSection(const FAssetData& AssetData, UMovieSceneMe
 	}
 }
 
+void FMediaTrackEditor::AddNewSectionEnterPressed(const TArray<FAssetData>& AssetData, UMovieSceneMediaTrack* Track)
+{
+	if (AssetData.Num() > 0)
+	{
+		AddNewSection(AssetData[0].GetAsset(), Track);
+	}
+}
 
 /* FMediaTrackEditor callbacks
  *****************************************************************************/
diff --git a/Engine/Plugins/Media/MfMedia/Source/MfMedia/Private/Mf/MfMediaUtils.cpp b/Engine/Plugins/Media/MfMedia/Source/MfMedia/Private/Mf/MfMediaUtils.cpp
--- a/Engine/Plugins/Media/MfMedia/Source/MfMedia/Private/Mf/MfMediaUtils.cpp
+++ b/Engine/Plugins/Media/MfMedia/Source/MfMedia/Private/Mf/MfMediaUtils.cpp
@@ -143,17 +143,27 @@ namespace MfMedia
 				}
 			}
 
-			if ((SubType == MFVideoFormat_HEVC) && !FWindowsPlatformMisc::VerifyWindowsVersion(10, 0))
+			if ((SubType == MFVideoFormat_H264) || (SubType == MFVideoFormat_H264_ES))
 			{
-				UE_LOG(LogMfMedia, Warning, TEXT("Your Windows version is %s"), *FPlatformMisc::GetOSVersion());
-
-				if ((SubType == MFVideoFormat_HEVC) && !FWindowsPlatformMisc::VerifyWindowsVersion(6, 2))
+				if (!FWindowsPlatformMisc::VerifyWindowsVersion(6, 1) /*Win7*/)
 				{
-					UE_LOG(LogMfMedia, Warning, TEXT("HEVC video type requires Windows 10 or newer"));
+					UE_LOG(LogMfMedia, Warning, TEXT("H264 video type requires Windows 7 or newer (your version is %s)"), *FPlatformMisc::GetOSVersion());
 					return NULL;
 				}
+			}
+
+			if ((SubType == MFVideoFormat_HEVC) || (SubType == MFVideoFormat_HEVC_ES))
+			{
+				if (!FWindowsPlatformMisc::VerifyWindowsVersion(10, 0) /*Win10*/)
+				{
+					if (!FWindowsPlatformMisc::VerifyWindowsVersion(6, 2) /*Win8*/)
+					{
+						UE_LOG(LogMfMedia, Warning, TEXT("HEVC video type requires Windows 10 or newer (your version is %s"), *FPlatformMisc::GetOSVersion());
+						return NULL;
+					}
 
-				UE_LOG(LogMfMedia, Warning, TEXT("HEVC video type requires Windows 10 or newer (game must be manifested for Windows 10)"));
+					UE_LOG(LogMfMedia, Warning, TEXT("HEVC video type requires Windows 10 or newer (your version is %s), and game must be manifested for Windows 10"), *FPlatformMisc::GetOSVersion());
+				}
 			}
 #endif //PLATFORM_XBOXONE
 
@@ -169,7 +179,24 @@ namespace MfMedia
 #if PLATFORM_XBOXONE
 			Result = OutputType->SetGUID(MF_MT_SUBTYPE, MFVideoFormat_NV12); // XboxOne only supports NV12
 #else
-			Result = OutputType->SetGUID(MF_MT_SUBTYPE, (SubType == MFVideoFormat_RGB32) ? MFVideoFormat_RGB32 : MFVideoFormat_NV12);
+			if ((SubType == MFVideoFormat_HEVC) ||
+				(SubType == MFVideoFormat_HEVC_ES) ||
+				(SubType == MFVideoFormat_NV12) ||
+				(SubType == MFVideoFormat_IYUV))
+			{
+				Result = OutputType->SetGUID(MF_MT_SUBTYPE, MFVideoFormat_NV12);
+			}
+			else
+			{
+				const bool Uncompressed =
+					(SubType == MFVideoFormat_RGB555) ||
+					(SubType == MFVideoFormat_RGB565) ||
+					(SubType == MFVideoFormat_RGB24) ||
+					(SubType == MFVideoFormat_RGB32) ||
+					(SubType == MFVideoFormat_ARGB32);
+
+				Result = OutputType->SetGUID(MF_MT_SUBTYPE, Uncompressed ? MFVideoFormat_RGB32 : MFVideoFormat_YUY2);
+			}
 #endif
 
 			if (FAILED(Result))
diff --git a/Engine/Plugins/Media/MfMedia/Source/MfMedia/Private/Player/MfMediaTracks.cpp b/Engine/Plugins/Media/MfMedia/Source/MfMedia/Private/Player/MfMediaTracks.cpp
--- a/Engine/Plugins/Media/MfMedia/Source/MfMedia/Private/Player/MfMediaTracks.cpp
+++ b/Engine/Plugins/Media/MfMedia/Source/MfMedia/Private/Player/MfMediaTracks.cpp
@@ -1156,48 +1156,31 @@ bool FMfMediaTracks::AddStreamToTracks(uint32 StreamIndex, FString& OutInfo)
 					SampleFormat = EMediaTextureSampleFormat::CharNV12;
 				}
 #if PLATFORM_WINDOWS
+				else if (OutputSubType == MFVideoFormat_RGB32)
+				{
+					BufferDim = OutputDim;
+					BufferStride = OutputDim.X * 4;
+					SampleFormat = EMediaTextureSampleFormat::CharBMP;
+				}
 				else
 				{
-					long SampleStride = ::MFGetAttributeUINT32(OutputType, MF_MT_DEFAULT_STRIDE, 0);
+					int32 AlignedOutputX = OutputDim.X;
 
-					if (OutputSubType == MFVideoFormat_RGB32)
+					if ((SubType == MFVideoFormat_H264) || (SubType == MFVideoFormat_H264_ES))
 					{
-						SampleFormat = EMediaTextureSampleFormat::CharBMP;
-
-						if (SampleStride == 0)
-						{
-							::MFGetStrideForBitmapInfoHeader(OutputSubType.Data1, OutputDim.X, &SampleStride);
-						}
-
-						if (SampleStride == 0)
-						{
-							SampleStride = OutputDim.X * 4;
-						}
+						AlignedOutputX = Align(AlignedOutputX, 16);
 					}
-					else
-					{
-						SampleFormat = EMediaTextureSampleFormat::CharYUY2;
-
-						if (SampleStride == 0)
-						{
-							int32 AlignedOutputX = OutputDim.X;
 
-							if ((SubType == MFVideoFormat_H264) || (SubType == MFVideoFormat_H264_ES))
-							{
-								AlignedOutputX = Align(AlignedOutputX, 16);
-							}
-
-							SampleStride = AlignedOutputX * 2;
-						}
-					}
+					int32 SampleStride = AlignedOutputX * 2; // 2 bytes per pixel
 
 					if (SampleStride < 0)
 					{
 						SampleStride = -SampleStride;
 					}
 
-					BufferDim = FIntPoint(SampleStride / 4, OutputDim.Y);
+					BufferDim = FIntPoint(AlignedOutputX / 2, OutputDim.Y); // 2 pixels per texel
 					BufferStride = SampleStride;
+					SampleFormat = EMediaTextureSampleFormat::CharYUY2;
 				}
 #endif //PLATFORM_WINDOWS
 			}
diff --git a/Engine/Plugins/Media/WmfMedia/Source/WmfMedia/Private/Player/WmfMediaTracks.cpp b/Engine/Plugins/Media/WmfMedia/Source/WmfMedia/Private/Player/WmfMediaTracks.cpp
--- a/Engine/Plugins/Media/WmfMedia/Source/WmfMedia/Private/Player/WmfMediaTracks.cpp
+++ b/Engine/Plugins/Media/WmfMedia/Source/WmfMedia/Private/Player/WmfMediaTracks.cpp
@@ -1052,6 +1052,8 @@ bool FWmfMediaTracks::AddTrackToTopology(const FTrack& Track, IMFTopology& Topol
 		return false;
 	}
 
+	UE_LOG(LogWmfMedia, Verbose, TEXT("Tracks %p: Added stream %i to topology"), this, Track.StreamIndex);
+
 	return true;
 }
 
diff --git a/Engine/Plugins/Media/WmfMedia/Source/WmfMedia/Private/Wmf/WmfMediaSession.cpp b/Engine/Plugins/Media/WmfMedia/Source/WmfMedia/Private/Wmf/WmfMediaSession.cpp
--- a/Engine/Plugins/Media/WmfMedia/Source/WmfMedia/Private/Wmf/WmfMediaSession.cpp
+++ b/Engine/Plugins/Media/WmfMedia/Source/WmfMedia/Private/Wmf/WmfMediaSession.cpp
@@ -183,7 +183,7 @@ bool FWmfMediaSession::SetTopology(const TComPtr<IMFTopology>& InTopology, FTime
 
 			if (FAILED(Result))
 			{
-				UE_LOG(LogWmfMedia, Verbose, TEXT("Session %p: Failed to set topology: %s"), this, *WmfMedia::ResultToString(Result));
+				UE_LOG(LogWmfMedia, Verbose, TEXT("Session %p: Failed to set partial topology %p: %s"), this, InTopology.Get(), *WmfMedia::ResultToString(Result));
 				
 				SessionState = EMediaState::Error;
 				DeferredEvents.Enqueue(EMediaEvent::MediaOpenFailed);
@@ -947,7 +947,7 @@ bool FWmfMediaSession::CommitTopology(IMFTopology* Topology)
 
 		if (FAILED(Result))
 		{
-			UE_LOG(LogWmfMedia, Verbose, TEXT("Session %p: Failed to set topology: %s"), this, *WmfMedia::ResultToString(Result));
+			UE_LOG(LogWmfMedia, Verbose, TEXT("Session %p: Failed to set topology %p: %s"), Topology, this, *WmfMedia::ResultToString(Result));
 			return false;
 		}
 
diff --git a/Engine/Plugins/Media/WmfMedia/Source/WmfMedia/Private/Wmf/WmfMediaUtils.cpp b/Engine/Plugins/Media/WmfMedia/Source/WmfMedia/Private/Wmf/WmfMediaUtils.cpp
--- a/Engine/Plugins/Media/WmfMedia/Source/WmfMedia/Private/Wmf/WmfMediaUtils.cpp
+++ b/Engine/Plugins/Media/WmfMedia/Source/WmfMedia/Private/Wmf/WmfMediaUtils.cpp
@@ -487,8 +487,7 @@ namespace WmfMedia
 			{
 				if (IsVideoDevice /*&& !FWindowsPlatformMisc::VerifyWindowsVersion(6, 2)*/ /*Win8*/)
 				{
-					UE_LOG(LogWmfMedia, Warning, TEXT("Your Windows version is %s"), *FPlatformMisc::GetOSVersion());
-					UE_LOG(LogWmfMedia, Warning, TEXT("H264 video type requires Windows 8 or newer"));
+					UE_LOG(LogWmfMedia, Warning, TEXT("H264 video type requires Windows 8 or newer (your version is %s)"), *FPlatformMisc::GetOSVersion());
 					return NULL;
 				}
 			}
@@ -497,15 +496,13 @@ namespace WmfMedia
 			{
 				if (!FWindowsPlatformMisc::VerifyWindowsVersion(10, 0) /*Win10*/)
 				{
-					UE_LOG(LogWmfMedia, Warning, TEXT("Your Windows version is %s"), *FPlatformMisc::GetOSVersion());
-
 					if (!FWindowsPlatformMisc::VerifyWindowsVersion(6, 2) /*Win8*/)
 					{
-						UE_LOG(LogWmfMedia, Warning, TEXT("HEVC video type requires Windows 10 or newer"));
+						UE_LOG(LogWmfMedia, Warning, TEXT("HEVC video type requires Windows 10 or newer (your version is %s)"), *FPlatformMisc::GetOSVersion());
 						return NULL;
 					}
 
-					UE_LOG(LogWmfMedia, Warning, TEXT("HEVC video type requires Windows 10 or newer (game must be manifested for Windows 10)"));
+					UE_LOG(LogWmfMedia, Warning, TEXT("HEVC video type requires Windows 10 or newer (your version is %s), and game must be manifested for Windows 10"), *FPlatformMisc::GetOSVersion());
 				}
 			}
 
diff --git a/Engine/Plugins/Messaging/UdpMessaging/Source/UdpMessaging/Private/Tests/UdpMessageTransportTest.cpp b/Engine/Plugins/Messaging/UdpMessaging/Source/UdpMessaging/Private/Tests/UdpMessageTransportTest.cpp
--- a/Engine/Plugins/Messaging/UdpMessaging/Source/UdpMessaging/Private/Tests/UdpMessageTransportTest.cpp
+++ b/Engine/Plugins/Messaging/UdpMessaging/Source/UdpMessaging/Private/Tests/UdpMessageTransportTest.cpp
@@ -103,62 +103,65 @@ bool FUdpMessageTransportTest::RunTest(const FString& Parameters)
 	const auto& DiscoveredNodes2 = Transport2.GetDiscoveredNodes();
 
 	// test transport node discovery
-	Transport1.Start();
-	FPlatformProcess::Sleep(3.0f);
-
-	TestTrue(TEXT("A single message transport must not discover any remote nodes"), DiscoveredNodes1.Num() == 0);
-
-	Transport2.Start();
-	FPlatformProcess::Sleep(3.0f);
-
-	if (DiscoveredNodes1.Num() == 0)
 	{
-		AddError(TEXT("The first transport did not discover any nodes"));
-
-		return false;
-	}
-
-	if (DiscoveredNodes2.Num() == 0)
-	{
-		AddError(TEXT("The second transport did not discover any nodes"));
-
-		return false;
+		Transport1.Start();
+		FPlatformProcess::Sleep(3.0f);
+
+		TestTrue(TEXT("A single message transport must not discover any remote nodes"), DiscoveredNodes1.Num() == 0);
+
+		Transport2.Start();
+		FPlatformProcess::Sleep(3.0f);
+
+		if (DiscoveredNodes1.Num() == 0)
+		{
+			AddError(TEXT("The first transport did not discover any nodes"));
+			return false;
+		}
+
+		if (DiscoveredNodes2.Num() == 0)
+		{
+			AddError(TEXT("The second transport did not discover any nodes"));
+			return false;
+		}
+
+		TestTrue(TEXT("The first transport must discover exactly one node"), DiscoveredNodes1.Num() == 1);
+		TestTrue(TEXT("The second transport must discover exactly one node"), DiscoveredNodes2.Num() == 1);
+		TestTrue(TEXT("The discovered node IDs must be valid"), DiscoveredNodes1[0].IsValid() && DiscoveredNodes2[0].IsValid());
+		TestTrue(TEXT("The discovered node IDs must be unique"), DiscoveredNodes1[0] != DiscoveredNodes2[0]);
 	}
 
-	TestTrue(TEXT("The first transport must discover exactly one node"), DiscoveredNodes1.Num() == 1);
-	TestTrue(TEXT("The second transport must discover exactly one node"), DiscoveredNodes2.Num() == 1);
-	TestTrue(TEXT("The discovered node IDs must be valid"), DiscoveredNodes1[0].IsValid() && DiscoveredNodes2[0].IsValid());
-	TestTrue(TEXT("The discovered node IDs must be unique"), DiscoveredNodes1[0] != DiscoveredNodes2[0]);
-
 	if (HasAnyErrors())
 	{
 		return false;
 	}
 
 	// stress test message sending
-	FDateTime StartTime = FDateTime::UtcNow();
-
-	for (int32 Count = 0; Count < NumTestMessages; ++Count)
 	{
-		FUdpMockMessage* Message = new FUdpMockMessage(MessageSize);
-		TSharedRef<IMessageContext, ESPMode::ThreadSafe> Context = MakeShareable(new FUdpMockMessageContext(Message));
+		const FDateTime StartTime = FDateTime::UtcNow();
 
-		Transport1.Publish(Context);
-	}
+		for (int32 Count = 0; Count < NumTestMessages; ++Count)
+		{
+			FUdpMockMessage* Message = new FUdpMockMessage(MessageSize);
+			TSharedRef<IMessageContext, ESPMode::ThreadSafe> Context = MakeShareable(new FUdpMockMessageContext(Message, StartTime));
 
-	AddInfo(FString::Printf(TEXT("Sent %i messages in %s"), NumTestMessages, *(FDateTime::UtcNow() - StartTime).ToString()));
+			Transport1.Publish(Context);
+		}
 
-	while ((Transport2.GetNumReceivedMessages() < NumTestMessages) && ((FDateTime::UtcNow() - StartTime) < FTimespan::FromSeconds(120.0)))
-	{
-		FPlatformProcess::Sleep(0.0f);
-	}
+		AddInfo(FString::Printf(TEXT("Sent %i messages in %s"), NumTestMessages, *(FDateTime::UtcNow() - StartTime).ToString()));
+
+		while ((Transport2.GetNumReceivedMessages() < NumTestMessages) && ((FDateTime::UtcNow() - StartTime) < FTimespan::FromSeconds(120.0)))
+		{
+			FPlatformProcess::Sleep(0.0f);
+		}
 
-	AddInfo(FString::Printf(TEXT("Received %i messages in %s"), Transport2.GetNumReceivedMessages(), *(FDateTime::UtcNow() - StartTime).ToString()));
-	TestTrue(TEXT("All sent messages must have been received"), Transport2.GetNumReceivedMessages() == NumTestMessages);
+		AddInfo(FString::Printf(TEXT("Received %i messages in %s"), Transport2.GetNumReceivedMessages(), *(FDateTime::UtcNow() - StartTime).ToString()));
+		TestTrue(TEXT("All sent messages must have been received"), Transport2.GetNumReceivedMessages() == NumTestMessages);
+	}
 
 	return true;
 }
 
+
 void EmptyLinkFunctionForStaticInitializationUdpMessageTransportTest()
 {
 	// This function exists to prevent the object file containing this test from being excluded by the linker, because it has no publicly referenced symbols.
diff --git a/Engine/Plugins/Messaging/UdpMessaging/Source/UdpMessaging/Private/Tests/UdpMessagingTestTypes.h b/Engine/Plugins/Messaging/UdpMessaging/Source/UdpMessaging/Private/Tests/UdpMessagingTestTypes.h
--- a/Engine/Plugins/Messaging/UdpMessaging/Source/UdpMessaging/Private/Tests/UdpMessagingTestTypes.h
+++ b/Engine/Plugins/Messaging/UdpMessaging/Source/UdpMessaging/Private/Tests/UdpMessagingTestTypes.h
@@ -33,12 +33,12 @@ class FUdpMockMessageContext
 {
 public:
 
-	FUdpMockMessageContext(FUdpMockMessage* InMessage)
+	FUdpMockMessageContext(FUdpMockMessage* InMessage, const FDateTime& InTimeSent)
 		: Expiration(FDateTime::MaxValue())
 		, Message(InMessage)
 		, Scope(EMessageScope::Network)
 		, SenderThread(ENamedThreads::AnyThread)
-		, TimeSent(FDateTime(2015, 9, 17, 10, 59, 23, 666))
+		, TimeSent(InTimeSent)
 		, TypeInfo(FUdpMockMessage::StaticStruct())
 	{
 		FMessageAddress::Parse(TEXT("11111111-22222222-33333333-44444444"), Sender);
@@ -59,7 +59,7 @@ public:
 
 public:
 
-	// IMessageContext interface
+	//~ IMessageContext interface
 
 	virtual const TMap<FName, FString>& GetAnnotations() const override { return Annotations; }
 	virtual TSharedPtr<IMessageAttachment, ESPMode::ThreadSafe> GetAttachment() const override { return Attachment; }
diff --git a/Engine/Plugins/Messaging/UdpMessaging/Source/UdpMessaging/Private/Tests/UdpSerializeMessageTaskTest.cpp b/Engine/Plugins/Messaging/UdpMessaging/Source/UdpMessaging/Private/Tests/UdpSerializeMessageTaskTest.cpp
--- a/Engine/Plugins/Messaging/UdpMessaging/Source/UdpMessaging/Private/Tests/UdpSerializeMessageTaskTest.cpp
+++ b/Engine/Plugins/Messaging/UdpMessaging/Source/UdpMessaging/Private/Tests/UdpSerializeMessageTaskTest.cpp
@@ -24,7 +24,8 @@ bool FUdpSerializeMessageTaskTest::RunTest(const FString& Parameters)
 {
 	using namespace UdpSerializeMessageTaskTest;
 
-	const auto Context = MakeShared<FUdpMockMessageContext, ESPMode::ThreadSafe>(new FUdpMockMessage);
+	const auto TimeSent = FDateTime(2015, 9, 17, 10, 59, 23, 666);
+	const auto Context = MakeShared<FUdpMockMessageContext, ESPMode::ThreadSafe>(new FUdpMockMessage, TimeSent);
 
 	// synchronous reference serialization
 	const auto Message1 = MakeShared<FUdpSerializedMessage, ESPMode::ThreadSafe>();
diff --git a/Engine/Plugins/MovieScene/ActorSequence/Source/ActorSequence/Private/ActorSequence.cpp b/Engine/Plugins/MovieScene/ActorSequence/Source/ActorSequence/Private/ActorSequence.cpp
--- a/Engine/Plugins/MovieScene/ActorSequence/Source/ActorSequence/Private/ActorSequence.cpp
+++ b/Engine/Plugins/MovieScene/ActorSequence/Source/ActorSequence/Private/ActorSequence.cpp
@@ -21,10 +21,10 @@ static TAutoConsoleVariable<int32> CVarDefaultEvaluationType(
 	TEXT("0: Playback locked to playback frames\n1: Unlocked playback with sub frame interpolation"),
 	ECVF_Default);
 
-static TAutoConsoleVariable<FString> CVarDefaultFrameResolution(
-	TEXT("ActorSequence.DefaultFrameResolution"),
+static TAutoConsoleVariable<FString> CVarDefaultTickResolution(
+	TEXT("ActorSequence.DefaultTickResolution"),
 	TEXT("24000fps"),
-	TEXT("Specifies default a frame resolution for newly created level sequences. Examples: 30 fps, 120/1 (120 fps), 30000/1001 (29.97), 0.01s (10ms)."),
+	TEXT("Specifies default a tick resolution for newly created level sequences. Examples: 30 fps, 120/1 (120 fps), 30000/1001 (29.97), 0.01s (10ms)."),
 	ECVF_Default);
 
 static TAutoConsoleVariable<FString> CVarDefaultDisplayRate(
@@ -86,13 +86,13 @@ void UActorSequence::PostInitProperties()
 
 		MovieScene->SetEvaluationType( bFrameLocked ? EMovieSceneEvaluationType::FrameLocked : EMovieSceneEvaluationType::WithSubFrames );
 
-		FFrameRate FrameResolution(60000, 1);
-		TryParseString(FrameResolution, *CVarDefaultFrameResolution.GetValueOnGameThread());
-		MovieScene->SetFrameResolutionDirectly(FrameResolution);
+		FFrameRate TickResolution(60000, 1);
+		TryParseString(TickResolution, *CVarDefaultTickResolution.GetValueOnGameThread());
+		MovieScene->SetTickResolutionDirectly(TickResolution);
 
-		FFrameRate PlaybackFrameRate(30, 1);
-		TryParseString(PlaybackFrameRate, *CVarDefaultDisplayRate.GetValueOnGameThread());
-		MovieScene->SetPlaybackFrameRate(PlaybackFrameRate);
+		FFrameRate DisplayRate(30, 1);
+		TryParseString(DisplayRate, *CVarDefaultDisplayRate.GetValueOnGameThread());
+		MovieScene->SetDisplayRate(DisplayRate);
 
 		OnInitializeSequenceEvent.Broadcast(this);
 		bHasBeenInitialized = true;
diff --git a/Engine/Plugins/MovieScene/ActorSequenceEditor/Source/ActorSequenceEditor/Private/ActorSequenceEditor.cpp b/Engine/Plugins/MovieScene/ActorSequenceEditor/Source/ActorSequenceEditor/Private/ActorSequenceEditor.cpp
--- a/Engine/Plugins/MovieScene/ActorSequenceEditor/Source/ActorSequenceEditor/Private/ActorSequenceEditor.cpp
+++ b/Engine/Plugins/MovieScene/ActorSequenceEditor/Source/ActorSequenceEditor/Private/ActorSequenceEditor.cpp
@@ -135,8 +135,8 @@ public:
 		auto* ProjectSettings = GetDefault<UMovieSceneToolsProjectSettings>();
 		UMovieScene* MovieScene = Sequence->GetMovieScene();
 		
-		FFrameNumber StartFrame = (ProjectSettings->DefaultStartTime * MovieScene->GetFrameResolution()).RoundToFrame();
-		int32        Duration   = (ProjectSettings->DefaultDuration * MovieScene->GetFrameResolution()).RoundToFrame().Value;
+		FFrameNumber StartFrame = (ProjectSettings->DefaultStartTime * MovieScene->GetTickResolution()).RoundToFrame();
+		int32        Duration   = (ProjectSettings->DefaultDuration * MovieScene->GetTickResolution()).RoundToFrame().Value;
 
 		MovieScene->SetPlaybackRange(StartFrame, Duration);
 	}
diff --git a/Engine/Plugins/MovieScene/LevelSequenceEditor/Source/LevelSequenceEditor/LevelSequenceEditor.Build.cs b/Engine/Plugins/MovieScene/LevelSequenceEditor/Source/LevelSequenceEditor/LevelSequenceEditor.Build.cs
--- a/Engine/Plugins/MovieScene/LevelSequenceEditor/Source/LevelSequenceEditor/LevelSequenceEditor.Build.cs
+++ b/Engine/Plugins/MovieScene/LevelSequenceEditor/Source/LevelSequenceEditor/LevelSequenceEditor.Build.cs
@@ -34,7 +34,8 @@ public class LevelSequenceEditor : ModuleRules
                 "Slate",
                 "SlateCore",
                 "UnrealEd",
-                "VREditor"
+                "VREditor",
+				"TimeManagement"
 			}
 		);
 
diff --git a/Engine/Plugins/MovieScene/LevelSequenceEditor/Source/LevelSequenceEditor/Private/CinematicViewport/SCinematicLevelViewport.cpp b/Engine/Plugins/MovieScene/LevelSequenceEditor/Source/LevelSequenceEditor/Private/CinematicViewport/SCinematicLevelViewport.cpp
--- a/Engine/Plugins/MovieScene/LevelSequenceEditor/Source/LevelSequenceEditor/Private/CinematicViewport/SCinematicLevelViewport.cpp
+++ b/Engine/Plugins/MovieScene/LevelSequenceEditor/Source/LevelSequenceEditor/Private/CinematicViewport/SCinematicLevelViewport.cpp
@@ -628,8 +628,8 @@ void SCinematicLevelViewport::Tick(const FGeometry& AllottedGeometry, const doub
 		SubTrack = Cast<UMovieSceneSubTrack>(Sequence->GetMovieScene()->FindMasterTrack(UMovieSceneSubTrack::StaticClass()));
 	}
 
-	const FFrameRate OuterResolution = Sequencer->GetFocusedFrameResolution();
-	const FFrameRate OuterPlayRate   = Sequencer->GetFocusedPlayRate();
+	const FFrameRate OuterResolution = Sequencer->GetFocusedTickResolution();
+	const FFrameRate OuterPlayRate   = Sequencer->GetFocusedDisplayRate();
 
 	const FFrameTime OuterTime       = Sequencer->GetLocalTime().ConvertTo(OuterResolution);
 	UIData.OuterResolution = OuterResolution;
@@ -655,7 +655,7 @@ void SCinematicLevelViewport::Tick(const FGeometry& AllottedGeometry, const doub
 	UMovieSceneSequence* SubSequence = SubSection ? SubSection->GetSequence() : nullptr;
 	if (SubSequence)
 	{
-		FFrameRate                   InnerResolution       = SubSequence->GetMovieScene()->GetFrameResolution();
+		FFrameRate                   InnerResolution       = SubSequence->GetMovieScene()->GetTickResolution();
 		FMovieSceneSequenceTransform OuterToInnerTransform = SubSection->OuterToInnerTransform();
 		const FFrameTime             InnerShotPosition	   = OuterTime * OuterToInnerTransform;
 
diff --git a/Engine/Plugins/MovieScene/LevelSequenceEditor/Source/LevelSequenceEditor/Private/CinematicViewport/SCinematicLevelViewport.h b/Engine/Plugins/MovieScene/LevelSequenceEditor/Source/LevelSequenceEditor/Private/CinematicViewport/SCinematicLevelViewport.h
--- a/Engine/Plugins/MovieScene/LevelSequenceEditor/Source/LevelSequenceEditor/Private/CinematicViewport/SCinematicLevelViewport.h
+++ b/Engine/Plugins/MovieScene/LevelSequenceEditor/Source/LevelSequenceEditor/Private/CinematicViewport/SCinematicLevelViewport.h
@@ -55,7 +55,7 @@ struct FUIData
 	/** The text that represents the master end frame */
 	FText MasterEndText;
 
-	/** The frame resolution of the master */
+	/** The tick resolution of the master */
 	FFrameRate OuterResolution;
 	/** The play rate of the master */
 	FFrameRate OuterPlayRate;
diff --git a/Engine/Plugins/MovieScene/LevelSequenceEditor/Source/LevelSequenceEditor/Private/CinematicViewport/SCinematicTransportRange.cpp b/Engine/Plugins/MovieScene/LevelSequenceEditor/Source/LevelSequenceEditor/Private/CinematicViewport/SCinematicTransportRange.cpp
--- a/Engine/Plugins/MovieScene/LevelSequenceEditor/Source/LevelSequenceEditor/Private/CinematicViewport/SCinematicTransportRange.cpp
+++ b/Engine/Plugins/MovieScene/LevelSequenceEditor/Source/LevelSequenceEditor/Private/CinematicViewport/SCinematicTransportRange.cpp
@@ -82,7 +82,7 @@ void SCinematicTransportRange::SetTime(const FGeometry& MyGeometry, const FPoint
 		FMovieSceneEditorData& EditorData = Sequencer->GetFocusedMovieSceneSequence()->GetMovieScene()->GetEditorData();
 		double NewTimeSeconds = EditorData.WorkStart + (EditorData.WorkEnd - EditorData.WorkStart) * Lerp;
 
-		Sequencer->SetLocalTime(NewTimeSeconds * Sequencer->GetFocusedFrameResolution(), ESnapTimeMode::STM_All);
+		Sequencer->SetLocalTime(NewTimeSeconds * Sequencer->GetFocusedTickResolution(), ESnapTimeMode::STM_All);
 	}
 }
 
@@ -99,7 +99,7 @@ void SCinematicTransportRange::Tick(const FGeometry& AllottedGeometry, const dou
 		TRange<double> WorkingRange = Sequencer->GetFocusedMovieSceneSequence()->GetMovieScene()->GetEditorData().GetWorkingRange();
 
 		// Anything within 3 pixel's worth of time is a duplicate as far as we're concerned
-		FTimeToPixel TimeToPixelConverter(AllottedGeometry, WorkingRange, Sequencer->GetFocusedFrameResolution());
+		FTimeToPixel TimeToPixelConverter(AllottedGeometry, WorkingRange, Sequencer->GetFocusedTickResolution());
 		const float DuplicateThreshold = (TimeToPixelConverter.PixelToSeconds(3.f) - TimeToPixelConverter.PixelToSeconds(0.f));
 
 		Sequencer->GetKeysFromSelection(ActiveKeyCollection, DuplicateThreshold);
@@ -119,11 +119,11 @@ int32 SCinematicTransportRange::OnPaint(const FPaintArgs& Args, const FGeometry&
 	static const float TrackOffsetY = 6.f;
 	const float TrackHeight = AllottedGeometry.GetLocalSize().Y - TrackOffsetY;
 
-	FFrameRate           FrameResolution = Sequencer->GetFocusedFrameResolution();
+	FFrameRate           TickResolution  = Sequencer->GetFocusedTickResolution();
 	TRange<double>       WorkingRange    = Sequencer->GetFocusedMovieSceneSequence()->GetMovieScene()->GetEditorData().GetWorkingRange();
 	TRange<FFrameNumber> PlaybackRange   = Sequencer->GetFocusedMovieSceneSequence()->GetMovieScene()->GetPlaybackRange();
 
-	const FFrameTime FramesPerPixel = (WorkingRange.Size<double>() / AllottedGeometry.GetLocalSize().X) * FrameResolution;
+	const FFrameTime FramesPerPixel = (WorkingRange.Size<double>() / AllottedGeometry.GetLocalSize().X) * TickResolution;
 
 	FColor DarkGray(40, 40, 40);
 	FColor MidGray(80, 80, 80);
@@ -141,8 +141,8 @@ int32 SCinematicTransportRange::OnPaint(const FPaintArgs& Args, const FGeometry&
 	
 	const float FullRange = WorkingRange.Size<float>();
 
-	const float PlaybackStartLerp	= (PlaybackRange.GetLowerBoundValue()/FrameResolution - WorkingRange.GetLowerBoundValue()) / FullRange;
-	const float PlaybackEndLerp		= (PlaybackRange.GetUpperBoundValue()/FrameResolution - WorkingRange.GetLowerBoundValue()) / FullRange;
+	const float PlaybackStartLerp	= (PlaybackRange.GetLowerBoundValue()/TickResolution - WorkingRange.GetLowerBoundValue()) / FullRange;
+	const float PlaybackEndLerp		= (PlaybackRange.GetUpperBoundValue()/TickResolution - WorkingRange.GetLowerBoundValue()) / FullRange;
 
 	// Draw the playback range
 	FSlateDrawElement::MakeBox(
@@ -186,8 +186,8 @@ int32 SCinematicTransportRange::OnPaint(const FPaintArgs& Args, const FGeometry&
 			const FSlateBrush* KeyBrush = FLevelSequenceEditorStyle::Get()->GetBrush("LevelSequenceEditor.CinematicViewportTransportRangeKey");
 
 			TRange<FFrameNumber> VisibleFrameRange(
-				(WorkingRange.GetLowerBoundValue() * FrameResolution).FloorToFrame(), 
-				(WorkingRange.GetUpperBoundValue() * FrameResolution).CeilToFrame()
+				(WorkingRange.GetLowerBoundValue() * TickResolution).FloorToFrame(), 
+				(WorkingRange.GetUpperBoundValue() * TickResolution).CeilToFrame()
 				);
 
 			for (FFrameNumber Time : ActiveKeyCollection->GetKeysInRange(VisibleFrameRange))
@@ -197,7 +197,7 @@ int32 SCinematicTransportRange::OnPaint(const FPaintArgs& Args, const FGeometry&
 					bPlayMarkerOnKey = true;
 				}
 
-				float Lerp = (Time/FrameResolution - WorkingRange.GetLowerBoundValue()) / FullRange;
+				float Lerp = (Time/TickResolution - WorkingRange.GetLowerBoundValue()) / FullRange;
 
 				FSlateDrawElement::MakeBox(
 					OutDrawElements,
diff --git a/Engine/Plugins/MovieScene/LevelSequenceEditor/Source/LevelSequenceEditor/Private/Factories/LevelSequenceFactoryNew.cpp b/Engine/Plugins/MovieScene/LevelSequenceEditor/Source/LevelSequenceEditor/Private/Factories/LevelSequenceFactoryNew.cpp
--- a/Engine/Plugins/MovieScene/LevelSequenceEditor/Source/LevelSequenceEditor/Private/Factories/LevelSequenceFactoryNew.cpp
+++ b/Engine/Plugins/MovieScene/LevelSequenceEditor/Source/LevelSequenceEditor/Private/Factories/LevelSequenceFactoryNew.cpp
@@ -31,8 +31,8 @@ UObject* ULevelSequenceFactoryNew::FactoryCreateNew(UClass* Class, UObject* InPa
 	// Set up some sensible defaults
 	const UMovieSceneToolsProjectSettings* ProjectSettings = GetDefault<UMovieSceneToolsProjectSettings>();
 
-	FFrameRate FrameResolution = NewLevelSequence->GetMovieScene()->GetFrameResolution();
-	NewLevelSequence->GetMovieScene()->SetPlaybackRange((ProjectSettings->DefaultStartTime*FrameResolution).FloorToFrame(), (ProjectSettings->DefaultDuration*FrameResolution).FloorToFrame().Value);
+	FFrameRate TickResolution = NewLevelSequence->GetMovieScene()->GetTickResolution();
+	NewLevelSequence->GetMovieScene()->SetPlaybackRange((ProjectSettings->DefaultStartTime*TickResolution).FloorToFrame(), (ProjectSettings->DefaultDuration*TickResolution).FloorToFrame().Value);
 
 	return NewLevelSequence;
 }
diff --git a/Engine/Plugins/MovieScene/LevelSequenceEditor/Source/LevelSequenceEditor/Private/LevelSequenceEditorToolkit.cpp b/Engine/Plugins/MovieScene/LevelSequenceEditor/Source/LevelSequenceEditor/Private/LevelSequenceEditorToolkit.cpp
--- a/Engine/Plugins/MovieScene/LevelSequenceEditor/Source/LevelSequenceEditor/Private/LevelSequenceEditorToolkit.cpp
+++ b/Engine/Plugins/MovieScene/LevelSequenceEditor/Source/LevelSequenceEditor/Private/LevelSequenceEditorToolkit.cpp
@@ -8,6 +8,8 @@
 #include "AssetData.h"
 #include "Components/PrimitiveComponent.h"
 #include "Editor.h"
+#include "Animation/AnimInstance.h"
+#include "Components/SkeletalMeshComponent.h"
 #include "Channels/MovieSceneChannelProxy.h"
 #include "Containers/ArrayBuilder.h"
 #include "Modules/ModuleManager.h"
@@ -244,6 +246,7 @@ void FLevelSequenceEditorToolkit::Initialize(const EToolkitMode::Type Mode, cons
 		SequencerInitParams.PlaybackContext.BindStatic(GetLevelSequenceEditorPlaybackContext);
 
 		SequencerInitParams.ViewParams.UniqueName = "LevelSequenceEditor";
+		SequencerInitParams.ViewParams.ScrubberStyle = ESequencerScrubberStyle::FrameBlock;
 		SequencerInitParams.ViewParams.OnReceivedFocus.BindRaw(this, &FLevelSequenceEditorToolkit::OnSequencerReceivedFocus);
 	}
 
@@ -287,10 +290,10 @@ void FLevelSequenceEditorToolkit::Initialize(const EToolkitMode::Type Mode, cons
 		{
 			VRMode->OnVREditingModeExit_Handler.BindSP(this, &FLevelSequenceEditorToolkit::HandleVREditorModeExit);
 			USequencerSettings& SavedSequencerSettings = *Sequencer->GetSequencerSettings();
-			VRMode->SaveSequencerSettings(Sequencer->GetKeyAllEnabled(), Sequencer->GetAutoChangeMode(), SavedSequencerSettings);
+			VRMode->SaveSequencerSettings(Sequencer->GetKeyGroupMode() == EKeyGroupMode::KeyAll, Sequencer->GetAutoChangeMode(), SavedSequencerSettings);
 			// Override currently set auto-change behavior to always autokey
 			Sequencer->SetAutoChangeMode(EAutoChangeMode::All);
-			Sequencer->SetKeyAllEnabled(true);
+			Sequencer->SetKeyGroupMode(EKeyGroupMode::KeyAll);
 			// Tell the VR Editor mode that Sequencer has refreshed
 			VRMode->RefreshVREditorSequencer(Sequencer.Get());
 		}
@@ -636,7 +639,7 @@ void FLevelSequenceEditorToolkit::HandleVREditorModeExit()
 
 	// Reset sequencer settings
 	Sequencer->SetAutoChangeMode(VRMode->GetSavedEditorState().AutoChangeMode);
-	Sequencer->SetKeyAllEnabled(VRMode->GetSavedEditorState().bKeyAllEnabled);
+	Sequencer->SetKeyGroupMode(VRMode->GetSavedEditorState().bKeyAllEnabled ? EKeyGroupMode::KeyAll : EKeyGroupMode::KeyChanged);
 	VRMode->OnVREditingModeExit_Handler.Unbind();
 }
 
@@ -765,13 +768,13 @@ void FLevelSequenceEditorToolkit::HandleMasterSequenceCreated(UObject* MasterSeq
 	UMovieSceneSequence* MasterSequence = Cast<UMovieSceneSequence>(MasterSequenceAsset);
 	UMovieSceneCinematicShotTrack* ShotTrack = MasterSequence->GetMovieScene()->AddMasterTrack<UMovieSceneCinematicShotTrack>();
 
-	FFrameRate FrameResolution = MasterSequence->GetMovieScene()->GetFrameResolution();
+	FFrameRate TickResolution = MasterSequence->GetMovieScene()->GetTickResolution();
 
 	// Create shots with a camera cut and a camera for each
-	FFrameNumber SequenceStartTime = (ProjectSettings->DefaultStartTime * FrameResolution).FloorToFrame();
+	FFrameNumber SequenceStartTime = (ProjectSettings->DefaultStartTime * TickResolution).FloorToFrame();
 	FFrameNumber ShotStartTime = SequenceStartTime;
 	FFrameNumber ShotEndTime   = ShotStartTime;
-	int32        ShotDuration  = (ProjectSettings->DefaultDuration * FrameResolution).RoundToFrame().Value;
+	int32        ShotDuration  = (ProjectSettings->DefaultDuration * TickResolution).RoundToFrame().Value;
 	FString FirstShotName; 
 	for (uint32 ShotIndex = 0; ShotIndex < NumShots; ++ShotIndex)
 	{
@@ -794,8 +797,8 @@ void FLevelSequenceEditorToolkit::HandleMasterSequenceCreated(UObject* MasterSeq
 	MasterSequence->GetMovieScene()->SetPlaybackRange(SequenceStartTime, (ShotEndTime - SequenceStartTime).Value);
 
 #if WITH_EDITORONLY_DATA
-	const double SequenceStartSeconds = SequenceStartTime / FrameResolution;
-	const double SequenceEndSeconds   = ShotEndTime / FrameResolution;
+	const double SequenceStartSeconds = SequenceStartTime / TickResolution;
+	const double SequenceEndSeconds   = ShotEndTime / TickResolution;
 	const double OutputChange = (SequenceEndSeconds - SequenceStartSeconds) * 0.1;
 
 	FMovieSceneEditorData& EditorData = MasterSequence->GetMovieScene()->GetEditorData();
@@ -882,17 +885,16 @@ void FLevelSequenceEditorToolkit::HandleTrackMenuExtensionAddTrack(FMenuBuilder&
 	}
 	else
 	{
-		UPrimitiveComponent* Component = Cast<UPrimitiveComponent>(ContextObjects[0]);
-		if (Component != nullptr)
+		if (UPrimitiveComponent* PrimitiveComponent = Cast<UPrimitiveComponent>(ContextObjects[0]))
 		{
-			int32 NumMaterials = Component->GetNumMaterials();
+			int32 NumMaterials = PrimitiveComponent->GetNumMaterials();
 			if (NumMaterials > 0)
 			{
 				AddTrackMenuBuilder.BeginSection("Materials", LOCTEXT("MaterialSection", "Materials"));
 				{
 					for (int32 MaterialIndex = 0; MaterialIndex < NumMaterials; MaterialIndex++)
 					{
-						FUIAction AddComponentMaterialAction(FExecuteAction::CreateSP(this, &FLevelSequenceEditorToolkit::HandleAddComponentMaterialActionExecute, Component, MaterialIndex));
+						FUIAction AddComponentMaterialAction(FExecuteAction::CreateSP(this, &FLevelSequenceEditorToolkit::HandleAddComponentMaterialActionExecute, PrimitiveComponent, MaterialIndex));
 						FText AddComponentMaterialLabel = FText::Format(LOCTEXT("ComponentMaterialIndexLabelFormat", "Element {0}"), FText::AsNumber(MaterialIndex));
 						FText AddComponentMaterialToolTip = FText::Format(LOCTEXT("ComponentMaterialIndexToolTipFormat", "Add material element {0}"), FText::AsNumber(MaterialIndex));
 						AddTrackMenuBuilder.AddMenuEntry(AddComponentMaterialLabel, AddComponentMaterialToolTip, FSlateIcon(), AddComponentMaterialAction);
@@ -901,9 +903,41 @@ void FLevelSequenceEditorToolkit::HandleTrackMenuExtensionAddTrack(FMenuBuilder&
 				AddTrackMenuBuilder.EndSection();
 			}
 		}
+
+		if (USkeletalMeshComponent* SkeletalComponent = Cast<USkeletalMeshComponent>(ContextObjects[0]))
+		{
+			UAnimInstance* AnimInstance = SkeletalComponent->GetAnimInstance();
+			
+			FText AnimInstanceLabel = LOCTEXT("AnimInstanceLabel", "Anim Instance");
+			FText DetailedInstanceText = AnimInstance
+				? FText::Format(LOCTEXT("AnimInstanceLabelFormat", "Anim Instance '{0}'"), FText::FromString(AnimInstance->GetName()))
+				: AnimInstanceLabel;
+
+			AddTrackMenuBuilder.BeginSection("Anim Instance", AnimInstanceLabel);
+			{
+				AddTrackMenuBuilder.AddMenuEntry(
+					DetailedInstanceText,
+					LOCTEXT("AnimInstanceToolTip", "Add this skeletal mesh component's animation instance."),
+					FSlateIcon(),
+					FUIAction(
+						FExecuteAction::CreateSP(this, &FLevelSequenceEditorToolkit::BindAnimationInstance, SkeletalComponent)
+					)
+				);
+			}
+			AddTrackMenuBuilder.EndSection();
+		}
 	}
 }
 
+void FLevelSequenceEditorToolkit::BindAnimationInstance(USkeletalMeshComponent* SkeletalComponent)
+{
+	UAnimInstance* AnimInstance = SkeletalComponent->GetAnimInstance();
+
+	// If there is no script instance at the moment, just use a dummy instance for the purposes of setting up the binding in the first place.
+	// This temporary object will get GC'd later on and is never actually applied to the anim instance
+	Sequencer->GetHandleToObject(AnimInstance ? AnimInstance : NewObject<UAnimInstance>(SkeletalComponent));
+}
+
 bool FLevelSequenceEditorToolkit::OnRequestClose()
 {
 	UWorld* World = CastChecked<UWorld>(GetLevelSequenceEditorPlaybackContext());
diff --git a/Engine/Plugins/MovieScene/LevelSequenceEditor/Source/LevelSequenceEditor/Private/LevelSequenceEditorToolkit.h b/Engine/Plugins/MovieScene/LevelSequenceEditor/Source/LevelSequenceEditor/Private/LevelSequenceEditorToolkit.h
--- a/Engine/Plugins/MovieScene/LevelSequenceEditor/Source/LevelSequenceEditor/Private/LevelSequenceEditorToolkit.h
+++ b/Engine/Plugins/MovieScene/LevelSequenceEditor/Source/LevelSequenceEditor/Private/LevelSequenceEditorToolkit.h
@@ -119,6 +119,9 @@ private:
 	/** Callback for executing the add component material track. */
 	void HandleAddComponentMaterialActionExecute(UPrimitiveComponent* Component, int32 MaterialIndex);
 
+	/** Create a new binding for the specified skeletal mesh component's animation instance. */
+	void BindAnimationInstance(USkeletalMeshComponent* SkeletalComponent);
+
 	/** Callback for map changes. */
 	void HandleMapChanged(UWorld* NewWorld, EMapChangeType MapChangeType);
 
diff --git a/Engine/Plugins/MovieScene/LevelSequenceEditor/Source/LevelSequenceEditor/Private/Misc/LevelSequenceEditorSpawnRegister.cpp b/Engine/Plugins/MovieScene/LevelSequenceEditor/Source/LevelSequenceEditor/Private/Misc/LevelSequenceEditorSpawnRegister.cpp
--- a/Engine/Plugins/MovieScene/LevelSequenceEditor/Source/LevelSequenceEditor/Private/Misc/LevelSequenceEditorSpawnRegister.cpp
+++ b/Engine/Plugins/MovieScene/LevelSequenceEditor/Source/LevelSequenceEditor/Private/Misc/LevelSequenceEditorSpawnRegister.cpp
@@ -30,12 +30,11 @@ FLevelSequenceEditorSpawnRegister::FLevelSequenceEditorSpawnRegister()
 	FLevelEditorModule& LevelEditor = FModuleManager::GetModuleChecked<FLevelEditorModule>("LevelEditor");
 	OnActorSelectionChangedHandle = LevelEditor.OnActorSelectionChanged().AddRaw(this, &FLevelSequenceEditorSpawnRegister::HandleActorSelectionChanged);
 
-	FAreObjectsEditable AreObjectsEditable = FAreObjectsEditable::CreateRaw(this, &FLevelSequenceEditorSpawnRegister::AreObjectsEditable);
-	OnAreObjectsEditableHandle = AreObjectsEditable.GetHandle();
-	LevelEditor.AddEditableObjectPredicate(AreObjectsEditable);
-
 #if WITH_EDITOR
 	GEditor->OnObjectsReplaced().AddRaw(this, &FLevelSequenceEditorSpawnRegister::OnObjectsReplaced);
+
+	OnObjectModifiedHandle = FCoreUObjectDelegates::OnObjectModified.AddRaw(this, &FLevelSequenceEditorSpawnRegister::OnObjectModified);
+	OnObjectSavedHandle    = FCoreUObjectDelegates::OnObjectSaved.AddRaw(this, &FLevelSequenceEditorSpawnRegister::OnPreObjectSaved);
 #endif
 }
 
@@ -45,7 +44,6 @@ FLevelSequenceEditorSpawnRegister::~FLevelSequenceEditorSpawnRegister()
 	if (FLevelEditorModule* LevelEditor = FModuleManager::GetModulePtr<FLevelEditorModule>("LevelEditor"))
 	{
 		LevelEditor->OnActorSelectionChanged().Remove(OnActorSelectionChangedHandle);
-		LevelEditor->RemoveEditableObjectPredicate(OnAreObjectsEditableHandle);
 	}
 
 	TSharedPtr<ISequencer> Sequencer = WeakSequencer.Pin();
@@ -57,6 +55,8 @@ FLevelSequenceEditorSpawnRegister::~FLevelSequenceEditorSpawnRegister()
 
 #if WITH_EDITOR
 	GEditor->OnObjectsReplaced().RemoveAll(this);
+	FCoreUObjectDelegates::OnObjectModified.Remove(OnObjectModifiedHandle);
+	FCoreUObjectDelegates::OnObjectSaved.Remove(OnObjectSavedHandle);
 #endif
 }
 
@@ -72,8 +72,8 @@ UObject* FLevelSequenceEditorSpawnRegister::SpawnObject(FMovieSceneSpawnable& Sp
 	
 	if (AActor* NewActor = Cast<AActor>(NewObject))
 	{
-		// Cache the spawned object, and editable state first
-		SpawnedObjects.FindOrAdd(TemplateID).Add(NewActor);
+		// Add an entry to the modified objects map to keep track of when this object has been modified
+		ModifiedObjects.Add(NewActor, FTrackedObjectState(TemplateID, Spawnable.GetGuid()));
 
 		// Select the actor if we think it should be selected
 		if (SelectedSpawnedObjects.Contains(FMovieSceneSpawnRegisterKey(TemplateID, Spawnable.GetGuid())))
@@ -89,10 +89,24 @@ UObject* FLevelSequenceEditorSpawnRegister::SpawnObject(FMovieSceneSpawnable& Sp
 void FLevelSequenceEditorSpawnRegister::PreDestroyObject(UObject& Object, const FGuid& BindingId, FMovieSceneSequenceIDRef TemplateID)
 {
 	TGuardValue<bool> Guard(bShouldClearSelectionCache, false);
-	
+
 	TSharedPtr<ISequencer> Sequencer = WeakSequencer.Pin();
 
-	SaveDefaultSpawnableState(BindingId, TemplateID);
+	UMovieSceneSequence*  Sequence      = Sequencer.IsValid() ? Sequencer->GetEvaluationTemplate().GetSequence(TemplateID) : nullptr;
+	FMovieSceneSpawnable* Spawnable     = Sequence ? Sequence->GetMovieScene()->FindSpawnable(BindingId) : nullptr;
+	UObject*              SpawnedObject = FindSpawnedObject(BindingId, TemplateID);
+
+	if (SpawnedObject && Spawnable)
+	{
+		const FTrackedObjectState* TrackedState = ModifiedObjects.Find(&Object);
+		if (TrackedState && TrackedState->bHasBeenModified)
+		{
+			// SaveDefaultSpawnableState will reset bHasBeenModified to false
+			SaveDefaultSpawnableStateImpl(*Spawnable, Sequence, SpawnedObject, *Sequencer);
+
+			Sequence->MarkPackageDirty();
+		}
+	}
 
 	// Cache its selection state
 	AActor* Actor = Cast<AActor>(&Object);
@@ -101,17 +115,9 @@ void FLevelSequenceEditorSpawnRegister::PreDestroyObject(UObject& Object, const
 		SelectedSpawnedObjects.Add(FMovieSceneSpawnRegisterKey(TemplateID, BindingId));
 		GEditor->SelectActor(Actor, false /*bSelected*/, true /*bNotify*/);
 	}
-	
-	// Remove the spawned object (and anything that's null) from our cache
-	TSet<FObjectKey>* ExistingObjects = SpawnedObjects.Find(TemplateID);
-	if (ExistingObjects)
-	{
-		ExistingObjects->Remove(FObjectKey(&Object));
-		if (ExistingObjects->Num() == 0)
-		{
-			SpawnedObjects.Remove(TemplateID);
-		}
-	}
+
+	FObjectKey ThisObject(&Object);
+	ModifiedObjects.Remove(ThisObject);
 
 	FLevelSequenceSpawnRegister::PreDestroyObject(Object, BindingId, TemplateID);
 }
@@ -121,30 +127,11 @@ void FLevelSequenceEditorSpawnRegister::SaveDefaultSpawnableState(FMovieSceneSpa
 	UMovieSceneSequence* Sequence = Player.GetEvaluationTemplate().GetSequence(TemplateID);
 
 	UObject* Object = FindSpawnedObject(Spawnable.GetGuid(), TemplateID);
-	if (!Object || !Sequence)
+	if (Object && Sequence)
 	{
-		return;
+		SaveDefaultSpawnableStateImpl(Spawnable, Sequence, Object, Player);
+		Sequence->MarkPackageDirty();
 	}
-
-	auto RestorePredicate = [](FMovieSceneAnimTypeID TypeID){ return TypeID != FMovieSceneSpawnSectionTemplate::GetAnimTypeID(); };
-
-	if (AActor* Actor = Cast<AActor>(Object))
-	{
-		// Restore state on any components
-		for (UActorComponent* Component : TInlineComponentArray<UActorComponent*>(Actor))
-		{
-			if (Component)
-			{
-				Player.RestorePreAnimatedState(*Component, RestorePredicate);
-			}
-		}
-	}
-
-	// Restore state on the object itself
-	Player.RestorePreAnimatedState(*Object, RestorePredicate);
-
-	// Copy the template
-	Spawnable.CopyObjectTemplate(*Object, *Sequence);
 }
 
 void FLevelSequenceEditorSpawnRegister::SaveDefaultSpawnableState(const FGuid& BindingId, FMovieSceneSequenceIDRef TemplateID)
@@ -171,43 +158,49 @@ void FLevelSequenceEditorSpawnRegister::SaveDefaultSpawnableState(const FGuid& B
 
 	if (Spawnable)
 	{
-		SaveDefaultSpawnableState(*Spawnable, TemplateID, *Sequencer);
+		UObject* Object = FindSpawnedObject(Spawnable->GetGuid(), TemplateID);
+		if (Object)
+		{
+			SaveDefaultSpawnableStateImpl(*Spawnable, Sequence, Object, *Sequencer);
+			Sequence->MarkPackageDirty();
+		}
 	}
 }
 
-void FLevelSequenceEditorSpawnRegister::OnPreSaveMovieScene(ISequencer& InSequencer)
+void FLevelSequenceEditorSpawnRegister::SaveDefaultSpawnableStateImpl(FMovieSceneSpawnable& Spawnable, UMovieSceneSequence* Sequence, UObject* SpawnedObject, IMovieScenePlayer& Player)
 {
-	// We're about to save the movie scene(s), so we need to save default spawnable state for the currently focused movie scene sequence instance
+	auto RestorePredicate = [](FMovieSceneAnimTypeID TypeID){ return TypeID != FMovieSceneSpawnSectionTemplate::GetAnimTypeID(); };
 
-	UMovieSceneSequence* Sequence = InSequencer.GetFocusedMovieSceneSequence();
-	UMovieScene* MovieScene = Sequence ? Sequence->GetMovieScene() : nullptr;
-	if (!MovieScene)
+	if (AActor* Actor = Cast<AActor>(SpawnedObject))
 	{
-		return;
+		// Restore state on any components
+		for (UActorComponent* Component : TInlineComponentArray<UActorComponent*>(Actor))
+		{
+			if (Component)
+			{
+				Player.RestorePreAnimatedState(*Component, RestorePredicate);
+			}
+		}
 	}
 
-	for (int32 SpawnableIndex = 0; SpawnableIndex < MovieScene->GetSpawnableCount(); ++SpawnableIndex)
+	// Restore state on the object itself
+	Player.RestorePreAnimatedState(*SpawnedObject, RestorePredicate);
+
+	// Copy the template
+	Spawnable.CopyObjectTemplate(*SpawnedObject, *Sequence);
+
+	if (FTrackedObjectState* TrackedState = ModifiedObjects.Find(SpawnedObject))
 	{
-		FMovieSceneSpawnable& Spawnable = MovieScene->GetSpawnable(SpawnableIndex);
-		SaveDefaultSpawnableState(Spawnable.GetGuid(), InSequencer.GetFocusedTemplateID());
+		TrackedState->bHasBeenModified = false;
 	}
 }
 
-void FLevelSequenceEditorSpawnRegister::OnSequenceInstanceActivated(FMovieSceneSequenceIDRef InTemplateID)
-{
-	ActiveSequence = InTemplateID;
-}
-
 /* FLevelSequenceEditorSpawnRegister implementation
  *****************************************************************************/
 
 void FLevelSequenceEditorSpawnRegister::SetSequencer(const TSharedPtr<ISequencer>& Sequencer)
 {
 	WeakSequencer = Sequencer;
-	Sequencer->OnPreSave().AddRaw(this, &FLevelSequenceEditorSpawnRegister::OnPreSaveMovieScene);
-	Sequencer->OnActivateSequence().AddRaw(this, &FLevelSequenceEditorSpawnRegister::OnSequenceInstanceActivated);
-
-	ActiveSequence = Sequencer->GetFocusedTemplateID();
 }
 
 
@@ -218,63 +211,10 @@ void FLevelSequenceEditorSpawnRegister::HandleActorSelectionChanged(const TArray
 {
 	if (bShouldClearSelectionCache)
 	{
-		// @todo: arodham: how does this merge? SpawnedObject.Key relates to the old unique sequence instance ID, not the object binding ID. Not sure what's intended here. Source Ref: CL#3117184
-		// TSharedPtr<ISequencer> Sequencer = WeakSequencer.Pin();
-		// if (Sequencer.IsValid())
-		// {
-		// 	for (auto SpawnedObject : SpawnedObjects)
-		// 	{
-		// 		FMovieSceneSpawnRegisterKey SpawnRegisterKey(SpawnedObject.Key, Sequencer->GetFocusedMovieSceneSequenceInstance().Get());
-		// 		if (SelectedSpawnedObjects.Contains(SpawnRegisterKey))
-		// 		{
-		// 			SelectedSpawnedObjects.Remove(SpawnRegisterKey);
-		// 			break;
-		// 		}
-		// 	}
-		// }
-
 		SelectedSpawnedObjects.Reset();
 	}
 }
 
-bool FLevelSequenceEditorSpawnRegister::AreObjectsEditable(const TArray<TWeakObjectPtr<UObject>>& InObjects) const
-{
-	// Check that none of the objects specified are (or belong to) a spawned actor
-	for (const TWeakObjectPtr<UObject>& WeakObject : InObjects)
-	{
-		UObject* Object = WeakObject.Get();
-		if (!Object)
-		{
-			continue;
-		}
-
-		AActor* SourceActor = Cast<AActor>(Object);
-		if (!SourceActor)
-		{
-			UActorComponent* ActorComponent = Cast<UActorComponent>(Object);
-			if (ActorComponent)
-			{
-				SourceActor = ActorComponent->GetOwner();
-			}
-		}
-
-		if (!SourceActor)
-		{
-			continue;
-		}
-
-		FObjectKey ThisObject(SourceActor);
-		for (auto& Pair : SpawnedObjects)
-		{
-			if (Pair.Key != ActiveSequence && Pair.Value.Contains(ThisObject))
-			{
-				return false;
-			}
-		}
-	}
-	return true;
-}
-
 void FLevelSequenceEditorSpawnRegister::OnObjectsReplaced(const TMap<UObject*, UObject*>& OldToNewInstanceMap)
 {
 	TSharedPtr<ISequencer> Sequencer = WeakSequencer.Pin();
@@ -301,6 +241,67 @@ void FLevelSequenceEditorSpawnRegister::OnObjectsReplaced(const TMap<UObject*, U
 	}
 }
 
+void FLevelSequenceEditorSpawnRegister::OnObjectModified(UObject* ModifiedObject)
+{
+	FTrackedObjectState* TrackedState = ModifiedObjects.Find(ModifiedObject);
+	while (!TrackedState && ModifiedObject)
+	{
+		TrackedState = ModifiedObjects.Find(ModifiedObject);
+		ModifiedObject = ModifiedObject->GetOuter();
+	}
+
+	if (TrackedState)
+	{
+		TrackedState->bHasBeenModified = true;
+
+		TSharedPtr<ISequencer> Sequencer = WeakSequencer.Pin();
+		UMovieSceneSequence*   OwningSequence = Sequencer.IsValid() ? Sequencer->GetEvaluationTemplate().GetSequence(TrackedState->TemplateID) : nullptr;
+		if (OwningSequence)
+		{
+			OwningSequence->MarkPackageDirty();
+			SequencesWithModifiedObjects.Add(OwningSequence);
+		}
+	}
+}
+
+void FLevelSequenceEditorSpawnRegister::OnPreObjectSaved(UObject* Object)
+{
+	UMovieSceneSequence* SequenceBeingSaved = Cast<UMovieSceneSequence>(Object);
+	if (SequenceBeingSaved && SequencesWithModifiedObjects.Contains(SequenceBeingSaved))
+	{
+		UMovieScene* MovieSceneBeingSaved = SequenceBeingSaved->GetMovieScene();
+
+		// The object being saved is a movie scene sequence that we've tracked as having a modified spawnable in the world.
+		// We need to go through all templates in the current sequence that reference this sequence, saving default state
+		// for any spawned objects that have been modified.
+		TSharedPtr<ISequencer> Sequencer = WeakSequencer.Pin();
+
+		if (Sequencer.IsValid())
+		{
+			TArray<FObjectKey> ObjectsThatHaveBeenSaved;
+
+			for (const TTuple<FObjectKey, FTrackedObjectState>& Pair : ModifiedObjects)
+			{
+				UObject* SpawnedObject = Pair.Key.ResolveObjectPtr();
+				UMovieSceneSequence*  ThisSequence = Sequencer->GetEvaluationTemplate().GetSequence(Pair.Value.TemplateID);
+				FMovieSceneSpawnable* Spawnable    = MovieSceneBeingSaved->FindSpawnable(Pair.Value.ObjectBindingID);
+
+				if (SpawnedObject && Spawnable && ThisSequence == SequenceBeingSaved)
+				{
+					SaveDefaultSpawnableStateImpl(*Spawnable, ThisSequence, SpawnedObject, *Sequencer);
+
+					ObjectsThatHaveBeenSaved.Add(Pair.Key);
+				}
+			}
+
+			for (FObjectKey ObjectKey : ObjectsThatHaveBeenSaved)
+			{
+				ModifiedObjects.Remove(ObjectKey);
+			}
+		}
+	}
+}
+
 #if WITH_EDITOR
 
 TValueOrError<FNewSpawnable, FText> FLevelSequenceEditorSpawnRegister::CreateNewSpawnableType(UObject& SourceObject, UMovieScene& OwnerMovieScene, UActorFactory* ActorFactory)
diff --git a/Engine/Plugins/MovieScene/LevelSequenceEditor/Source/LevelSequenceEditor/Private/Misc/LevelSequenceEditorSpawnRegister.h b/Engine/Plugins/MovieScene/LevelSequenceEditor/Source/LevelSequenceEditor/Private/Misc/LevelSequenceEditorSpawnRegister.h
--- a/Engine/Plugins/MovieScene/LevelSequenceEditor/Source/LevelSequenceEditor/Private/Misc/LevelSequenceEditorSpawnRegister.h
+++ b/Engine/Plugins/MovieScene/LevelSequenceEditor/Source/LevelSequenceEditor/Private/Misc/LevelSequenceEditorSpawnRegister.h
@@ -50,31 +50,48 @@ private:
 	/** Called when the editor selection has changed. */
 	void HandleActorSelectionChanged(const TArray<UObject*>& NewSelection, bool bForceRefresh);
 
-	/** Called before sequencer attempts to save the movie scene(s) it's editing */
-	void OnPreSaveMovieScene(ISequencer& InSequencer);
-
-	/** Called when a new movie scene sequence instance has been activated */
-	void OnSequenceInstanceActivated(FMovieSceneSequenceIDRef ActiveInstance);
-
 	/** Saves the default state for the specified spawnable, if an instance for it currently exists */
 	void SaveDefaultSpawnableState(const FGuid& BindingId, FMovieSceneSequenceIDRef TemplateID);
+	void SaveDefaultSpawnableStateImpl(FMovieSceneSpawnable& Spawnable, UMovieSceneSequence* Sequence, UObject* SpawnedObject, IMovieScenePlayer& Player);
 
-	/** Check whether the specified objects are editable on the details panel. Called from the level editor */
-	bool AreObjectsEditable(const TArray<TWeakObjectPtr<UObject>>& InObjects) const;
-	
 	/** Called from the editor when a blueprint object replacement has occurred */
 	void OnObjectsReplaced(const TMap<UObject*, UObject*>& OldToNewInstanceMap);
 
+	/** Called whenever an object is modified in the editor */
+	void OnObjectModified(UObject* ModifiedObject);
+
+	/** Called before an object is saved in the editor */
+	void OnPreObjectSaved(UObject* Object);
+
+private:
+
+	struct FTrackedObjectState
+	{
+		FTrackedObjectState(FMovieSceneSequenceIDRef InTemplateID, const FGuid& InObjectBindingID) : TemplateID(InTemplateID), ObjectBindingID(InObjectBindingID), bHasBeenModified(false) {}
+
+		/** The sequence ID that spawned this object */
+		FMovieSceneSequenceID TemplateID;
+
+		/** The object binding ID of the object in the template */
+		FGuid ObjectBindingID;
+
+		/** true if this object has been modified since it was spawned and is different from the current object template */
+		bool bHasBeenModified;
+	};
+
 private:
 
 	/** Handles for delegates that we've bound to. */
-	FDelegateHandle OnActorSelectionChangedHandle, OnAreObjectsEditableHandle;
+	FDelegateHandle OnActorSelectionChangedHandle;
 
 	/** Set of spawn register keys for objects that should be selected if they are spawned. */
 	TSet<FMovieSceneSpawnRegisterKey> SelectedSpawnedObjects;
 
-	/** Set of currently spawned objects */
-	TMap<FMovieSceneSequenceID, TSet<FObjectKey>> SpawnedObjects;
+	/** Map from a sequenceID to an array of objects that have been modified */
+	TMap<FObjectKey, FTrackedObjectState> ModifiedObjects;
+
+	/** Set of UMovieSceneSequences that this register has spawned objects for that are modified */
+	TSet<FObjectKey> SequencesWithModifiedObjects;
 
 	/** True if we should clear the above selection cache when the editor selection has been changed. */
 	bool bShouldClearSelectionCache;
@@ -82,6 +99,9 @@ private:
 	/** Weak pointer to the active sequencer. */
 	TWeakPtr<ISequencer> WeakSequencer;
 
-	/** Identifier for the current active level sequence */
-	FMovieSceneSequenceID ActiveSequence;
+	/** Handle to a delegate that is bound to FCoreUObjectDelegates::OnObjectModified to harvest changes to spawned objects. */
+	FDelegateHandle OnObjectModifiedHandle;
+
+	/** Handle to a delegate that is bound to FCoreUObjectDelegates::OnObjectSaved to harvest changes to spawned objects. */
+	FDelegateHandle OnObjectSavedHandle;
 };
diff --git a/Engine/Plugins/MovieScene/MatineeToLevelSequence/Source/MatineeToLevelSequence/Private/MatineeToLevelSequenceModule.cpp b/Engine/Plugins/MovieScene/MatineeToLevelSequence/Source/MatineeToLevelSequence/Private/MatineeToLevelSequenceModule.cpp
--- a/Engine/Plugins/MovieScene/MatineeToLevelSequence/Source/MatineeToLevelSequence/Private/MatineeToLevelSequenceModule.cpp
+++ b/Engine/Plugins/MovieScene/MatineeToLevelSequence/Source/MatineeToLevelSequence/Private/MatineeToLevelSequenceModule.cpp
@@ -698,7 +698,7 @@ protected:
 			MatineeActor->InitInterp();
 
 			// Set the length
-			const int32 LengthInFrames = (MatineeActor->MatineeData->InterpLength * NewMovieScene->GetFrameResolution()).FrameNumber.Value;
+			const int32 LengthInFrames = (MatineeActor->MatineeData->InterpLength * NewMovieScene->GetTickResolution()).FrameNumber.Value;
 			NewMovieScene->SetPlaybackRange(FFrameNumber(0), LengthInFrames + 1);
 
 			// Convert the groups
diff --git a/Engine/Plugins/Runtime/AudioCapture/Source/AudioCaptureEditor/Private/AudioRecordingManager.cpp b/Engine/Plugins/Runtime/AudioCapture/Source/AudioCaptureEditor/Private/AudioRecordingManager.cpp
--- a/Engine/Plugins/Runtime/AudioCapture/Source/AudioCaptureEditor/Private/AudioRecordingManager.cpp
+++ b/Engine/Plugins/Runtime/AudioCapture/Source/AudioCaptureEditor/Private/AudioRecordingManager.cpp
@@ -424,7 +424,8 @@ void FAudioRecordingManager::StopRecording(TArray<USoundWave*>& OutSoundWaves)
 					FMemory::Memcpy(NewSoundWave->RawPCMData, RawData, NumBytes);
 
 					// Calculate the duration of the sound wave
-					NewSoundWave->Duration = (float)(NumRecordedSamples / NumChannelsToSerialize) / WAVE_FILE_SAMPLERATE;
+					// Note: We use the NumInputChannels for duration calculation since NumChannelsToSerialize may be 1 channel while NumInputChannels is 2 for the "split stereo" feature.
+					NewSoundWave->Duration = (float)(NumRecordedSamples / NumInputChannels) / WAVE_FILE_SAMPLERATE;
 					NewSoundWave->SetSampleRate(WAVE_FILE_SAMPLERATE);
 					NewSoundWave->NumChannels = NumChannelsToSerialize;
 
diff --git a/Engine/Source/Editor/CurveEditor/Private/CurveEditor.cpp b/Engine/Source/Editor/CurveEditor/Private/CurveEditor.cpp
--- a/Engine/Source/Editor/CurveEditor/Private/CurveEditor.cpp
+++ b/Engine/Source/Editor/CurveEditor/Private/CurveEditor.cpp
@@ -307,8 +307,8 @@ void FCurveEditor::ZoomToFitSelection(EAxisList::Type Axes)
 void FCurveEditor::ZoomToFitInternal(EAxisList::Type Axes, double InputMin, double InputMax, double OutputMin, double OutputMax)
 {
 	FCurveEditorSnapMetrics SnapMetrics = GetSnapMetrics();
-	double MinInputZoom = SnapMetrics.bSnapInputValues ? SnapMetrics.InputSnapRate.AsInterval() : 0.001;
-	double MinOutputZoom = SnapMetrics.bSnapOutputValues ? SnapMetrics.OutputSnapInterval : 0.001;
+	double MinInputZoom = SnapMetrics.bSnapInputValues ? SnapMetrics.InputSnapRate.AsInterval() : 0.00001;
+	double MinOutputZoom = SnapMetrics.bSnapOutputValues ? SnapMetrics.OutputSnapInterval : 0.00001;
 
 	InputMax = FMath::Max(InputMin + MinInputZoom, InputMax);
 	OutputMax = FMath::Max(OutputMin + MinOutputZoom, OutputMax);
@@ -376,6 +376,30 @@ FVector2D FCurveEditor::GetVectorFromSlopeAndLength(float Slope, float Length)
 	return FVector2D(x, y);
 }
 
+FVector2D  FCurveEditor::GetTangentPositionInScreenSpace(const FVector2D &StartPos, float Tangent, float Weight) const
+{
+	FCurveEditorScreenSpace ScreenSpace = GetScreenSpace();
+	const float Angle = FMath::Atan(-Tangent);
+	float X, Y;
+	FMath::SinCos(&Y, &X, Angle);
+	X *= Weight;
+	Y *= Weight;
+
+	X *= ScreenSpace.PixelsPerInput();
+	Y *= ScreenSpace.PixelsPerOutput();
+	return FVector2D(StartPos.X + X, StartPos.Y +Y);
+}
+
+void FCurveEditor::GetTangentAndWeightFromScreenPosition(const FVector2D &StartPos, const  FVector2D &TangentPos, float &Tangent, float &Weight) const
+{
+	FCurveEditorScreenSpace ScreenSpace = GetScreenSpace();
+	float X = ScreenSpace.ScreenToSeconds(TangentPos.X) - ScreenSpace.ScreenToSeconds(StartPos.X);
+	float Y = ScreenSpace.ScreenToValue(TangentPos.Y) - ScreenSpace.ScreenToValue(StartPos.Y);
+
+	Tangent = Y / X;
+	Weight = FMath::Sqrt(X*X + Y * Y);
+}
+
 void FCurveEditor::GetCurveDrawParams(TArray<FCurveDrawParams>& OutDrawParams) const
 {
 	FCurveEditorScreenSpace ScreenSpace = GetScreenSpace();
@@ -456,7 +480,18 @@ void FCurveEditor::GetCurveDrawParams(TArray<FCurveDrawParams>& OutDrawParams) c
 
 					FCurvePointInfo ArriveTangentPoint(KeyHandle);
 					ArriveTangentPoint.Type = ECurvePointType::ArriveTangent;
-					ArriveTangentPoint.ScreenPosition = Key.ScreenPosition + FCurveEditor::GetVectorFromSlopeAndLength(ArriveTangent * -DisplayRatio, -60.f);
+
+
+					if (Attributes.HasTangentWeightMode() && Attributes.HasArriveTangentWeight() &&
+						(Attributes.GetTangentWeightMode() == RCTWM_WeightedBoth || Attributes.GetTangentWeightMode() == RCTWM_WeightedArrive))
+					{
+						ArriveTangentPoint.ScreenPosition = GetTangentPositionInScreenSpace(Key.ScreenPosition, ArriveTangent, -Attributes.GetArriveTangentWeight());
+					}
+					else
+					{
+						float PixelLength = 60.0f; 
+						ArriveTangentPoint.ScreenPosition = Key.ScreenPosition + FCurveEditor::GetVectorFromSlopeAndLength(ArriveTangent * -DisplayRatio, -PixelLength);
+					}
 					ArriveTangentPoint.LineDelta = Key.ScreenPosition - ArriveTangentPoint.ScreenPosition;
 					ArriveTangentPoint.LayerBias = 1;
 
@@ -469,7 +504,19 @@ void FCurveEditor::GetCurveDrawParams(TArray<FCurveDrawParams>& OutDrawParams) c
 
 					FCurvePointInfo LeaveTangentPoint(KeyHandle);
 					LeaveTangentPoint.Type = ECurvePointType::LeaveTangent;
-					LeaveTangentPoint.ScreenPosition = Key.ScreenPosition + FCurveEditor::GetVectorFromSlopeAndLength(LeaveTangent * -DisplayRatio, 60.f);
+
+					if (Attributes.HasTangentWeightMode() && Attributes.HasLeaveTangentWeight() &&
+						(Attributes.GetTangentWeightMode() == RCTWM_WeightedBoth || Attributes.GetTangentWeightMode() == RCTWM_WeightedLeave))
+					{
+						LeaveTangentPoint.ScreenPosition = GetTangentPositionInScreenSpace(Key.ScreenPosition, LeaveTangent, Attributes.GetLeaveTangentWeight());
+					}
+					else
+					{
+						float PixelLength = 60.0f; 
+						LeaveTangentPoint.ScreenPosition = Key.ScreenPosition + FCurveEditor::GetVectorFromSlopeAndLength(LeaveTangent * -DisplayRatio, PixelLength);
+
+					}
+
 					LeaveTangentPoint.LineDelta = Key.ScreenPosition - LeaveTangentPoint.ScreenPosition;
 					LeaveTangentPoint.LayerBias = 1;
 
diff --git a/Engine/Source/Editor/CurveEditor/Private/CurveEditorCommands.cpp b/Engine/Source/Editor/CurveEditor/Private/CurveEditorCommands.cpp
--- a/Engine/Source/Editor/CurveEditor/Private/CurveEditorCommands.cpp
+++ b/Engine/Source/Editor/CurveEditor/Private/CurveEditorCommands.cpp
@@ -19,6 +19,7 @@ void FCurveEditorCommands::RegisterCommands()
 	UI_COMMAND(InterpolationCubicAuto, "Auto", "Cubic interpolation - Automatic tangents", EUserInterfaceActionType::ToggleButton, FInputChord(EKeys::One));
 	UI_COMMAND(InterpolationCubicUser, "User", "Cubic interpolation - User flat tangents", EUserInterfaceActionType::ToggleButton, FInputChord(EKeys::Two));
 	UI_COMMAND(InterpolationCubicBreak, "Break", "Cubic interpolation - User broken tangents", EUserInterfaceActionType::ToggleButton, FInputChord(EKeys::Three));
+	UI_COMMAND(InterpolationToggleWeighted, "Weighted Tangents", "Toggle weighted tangents for cubic interpolation modes", EUserInterfaceActionType::ToggleButton, FInputChord(EModifierKey::Control, EKeys::W));
 
 	UI_COMMAND(FlattenTangents, "Flatten", "Flatten Tangents", EUserInterfaceActionType::Button, FInputChord());
 	UI_COMMAND(StraightenTangents, "Straighten", "Straighten Tangents", EUserInterfaceActionType::Button, FInputChord());
diff --git a/Engine/Source/Editor/CurveEditor/Private/DragOperations/CurveEditorDragOperation_Tangent.cpp b/Engine/Source/Editor/CurveEditor/Private/DragOperations/CurveEditorDragOperation_Tangent.cpp
--- a/Engine/Source/Editor/CurveEditor/Private/DragOperations/CurveEditorDragOperation_Tangent.cpp
+++ b/Engine/Source/Editor/CurveEditor/Private/DragOperations/CurveEditorDragOperation_Tangent.cpp
@@ -7,23 +7,6 @@
 void FCurveEditorDragOperation_Tangent::OnInitialize(FCurveEditor* InCurveEditor, const TOptional<FCurvePointHandle>& InCardinalPoint)
 {
 	CurveEditor = InCurveEditor;
-
-	// Find the position of the clicked on key (no type name)
-	if (InCardinalPoint.IsSet())
-	{
-		if (const FCurveModel* Curve = CurveEditor->FindCurve(InCardinalPoint->CurveID))
-		{
-			FCurveEditorScreenSpace ScreenSpace = CurveEditor->GetScreenSpace();
-
-			FKeyPosition CardinalKeyPosition;
-
-			Curve->GetKeyPositions(
-				TArrayView<const FKeyHandle>(&InCardinalPoint->KeyHandle, 1),
-				MakeArrayView(&CardinalKeyPosition, 1));
-
-			CardinalPosition = FVector2D(ScreenSpace.SecondsToScreen(CardinalKeyPosition.InputValue), ScreenSpace.ValueToScreen(CardinalKeyPosition.OutputValue));
-		}
-	}
 }
 
 void FCurveEditorDragOperation_Tangent::OnBeginDrag(FVector2D InitialPosition, FVector2D CurrentPosition, const FPointerEvent& MouseEvent)
@@ -63,33 +46,13 @@ void FCurveEditorDragOperation_Tangent::OnDrag(FVector2D InitialPosition, FVecto
 {
 	FCurveEditorScreenSpace ScreenSpace = CurveEditor->GetScreenSpace();
 	const float DisplayRatio = (ScreenSpace.PixelsPerOutput() / ScreenSpace.PixelsPerInput());
-
-	TOptional<float> NewTangent;
 	
-	FVector2D PixelDelta = CurrentPosition - CardinalPosition.Get(InitialPosition);
-
+	FVector2D PixelDelta = CurrentPosition - InitialPosition;
 	if (MouseEvent.IsShiftDown())
 	{
 		PixelDelta = RoundTrajectory(PixelDelta);
 	}
 
-	if (PointType == ECurvePointType::ArriveTangent)
-	{
-		PixelDelta.X = FMath::Min(PixelDelta.X, -KINDA_SMALL_NUMBER);
-	}
-	else
-	{
-		PixelDelta.X = FMath::Max(PixelDelta.X, KINDA_SMALL_NUMBER);
-	}
-	//If the CardinalPosition is set that means we are doing a left mouse and we are over the key(cardinal) position
-	//and we can set the tangent directly.
-	//If we are doing a middle mouse move, there is no cardinal position, but rather we need to apply the
-	//PixelDelta invidually to each tangent
-	if (CardinalPosition.IsSet())
-	{
-		NewTangent = (-PixelDelta.Y / PixelDelta.X) / DisplayRatio;
-	}
-
 	TArray<FKeyAttributes> NewKeyAttributesScratch;
 	TArray<FKeyPosition> KeyPositions;
 
@@ -100,30 +63,51 @@ void FCurveEditorDragOperation_Tangent::OnDrag(FVector2D InitialPosition, FVecto
 		{
 			NewKeyAttributesScratch.Reset();
 			NewKeyAttributesScratch.Reserve(KeyData.Attributes.Num());
-			if (!NewTangent.IsSet()) //No tangent, so MMB, so cache out positions for each key
-			{
-				KeyPositions.SetNumUninitialized(KeyData.Attributes.Num());
-				Curve->GetKeyPositions(KeyData.Handles, KeyPositions);
-			}
+			KeyPositions.SetNumUninitialized(KeyData.Attributes.Num());
+			Curve->GetKeyPositions(KeyData.Handles, KeyPositions);
+			
 
 			if (PointType == ECurvePointType::ArriveTangent)
 			{
-				for(int32 i =0; i < KeyData.Handles.Num();++i)
+				for (int32 i = 0; i < KeyData.Handles.Num(); ++i)
 				{
-					if (NewTangent.IsSet())
-					{
-						NewKeyAttributesScratch.Add(FKeyAttributes().SetArriveTangent(NewTangent.GetValue()));
-					}
-					else  //Apply Delta to each Tangent indvidually
+					const FKeyAttributes Attributes = KeyData.Attributes[i];
+					if (Attributes.HasArriveTangent())
 					{
-						const FKeyAttributes Attributes = KeyData.Attributes[i];
 						const float ArriveTangent = Attributes.GetArriveTangent();
 						const FVector2D  KeyScreenPosition = FVector2D(ScreenSpace.SecondsToScreen(KeyPositions[i].InputValue), ScreenSpace.ValueToScreen(KeyPositions[i].OutputValue));
-						FVector2D TangentScreenPosition = FCurveEditor::GetVectorFromSlopeAndLength(ArriveTangent * -DisplayRatio, -60.f) + KeyScreenPosition;
-						TangentScreenPosition += PixelDelta;
-						TangentScreenPosition -= KeyScreenPosition;
-						float Tangent = (-TangentScreenPosition.Y / TangentScreenPosition.X) / DisplayRatio;
-						NewKeyAttributesScratch.Add(FKeyAttributes().SetArriveTangent(Tangent));
+						if (Attributes.HasTangentWeightMode() && Attributes.HasArriveTangentWeight() &&
+							(Attributes.GetTangentWeightMode() == RCTWM_WeightedBoth || Attributes.GetTangentWeightMode() == RCTWM_WeightedArrive))
+						{
+							float Tangent, Weight;
+							FVector2D TangentScreenPosition = CurveEditor->GetTangentPositionInScreenSpace(KeyScreenPosition, ArriveTangent, -Attributes.GetArriveTangentWeight());
+							TangentScreenPosition.X += PixelDelta.X;
+							TangentScreenPosition.Y += PixelDelta.Y;
+
+							CurveEditor->GetTangentAndWeightFromScreenPosition(KeyScreenPosition, TangentScreenPosition, Tangent, Weight);
+
+							FKeyAttributes NewAttributes;
+
+							NewAttributes.SetArriveTangent(Tangent);
+							NewAttributes.SetArriveTangentWeight(Weight);
+							NewKeyAttributesScratch.Add(NewAttributes);
+
+						}
+						else
+						{
+							const float PixelLength = 60.0f;
+							FVector2D TangentScreenPosition = FCurveEditor::GetVectorFromSlopeAndLength(ArriveTangent * -DisplayRatio, -PixelLength) + KeyScreenPosition;
+							TangentScreenPosition += PixelDelta;
+							TangentScreenPosition -= KeyScreenPosition;
+							const float Tangent = (-TangentScreenPosition.Y / TangentScreenPosition.X) / DisplayRatio;
+							FKeyAttributes NewAttributes;
+							NewAttributes.SetArriveTangent(Tangent);
+							NewKeyAttributesScratch.Add(NewAttributes);
+						}
+					}
+					else //still need to add since expect attributes to equal num of selected
+					{
+						NewKeyAttributesScratch.Add(FKeyAttributes());
 					}
 				}
 			}
@@ -131,24 +115,43 @@ void FCurveEditorDragOperation_Tangent::OnDrag(FVector2D InitialPosition, FVecto
 			{
 				for (int32 i = 0; i < KeyData.Handles.Num(); ++i)
 				{
-					if (NewTangent.IsSet())
-					{
-						NewKeyAttributesScratch.Add(FKeyAttributes().SetLeaveTangent(NewTangent.GetValue()));
-					}
-					else  //Apply Delta to each Tangent indvidually
+					const FKeyAttributes Attributes = KeyData.Attributes[i];
+					if (Attributes.HasLeaveTangent())
 					{
-						const FKeyAttributes Attributes = KeyData.Attributes[i];
 						const float LeaveTangent = Attributes.GetLeaveTangent();
 						const FVector2D  KeyScreenPosition = FVector2D(ScreenSpace.SecondsToScreen(KeyPositions[i].InputValue), ScreenSpace.ValueToScreen(KeyPositions[i].OutputValue));
-						FVector2D TangentScreenPosition = FCurveEditor::GetVectorFromSlopeAndLength(LeaveTangent * -DisplayRatio,60.f) + KeyScreenPosition;
-						TangentScreenPosition += PixelDelta;
-						TangentScreenPosition -= KeyScreenPosition;
-						float Tangent = (-TangentScreenPosition.Y / TangentScreenPosition.X) / DisplayRatio;
-						NewKeyAttributesScratch.Add(FKeyAttributes().SetLeaveTangent(Tangent));
+						if (Attributes.HasTangentWeightMode() && Attributes.HasLeaveTangentWeight() &&
+							(Attributes.GetTangentWeightMode() == RCTWM_WeightedBoth || Attributes.GetTangentWeightMode() == RCTWM_WeightedLeave))
+						{
+							float Tangent, Weight;
+							FVector2D TangentScreenPosition = CurveEditor->GetTangentPositionInScreenSpace(KeyScreenPosition, LeaveTangent, Attributes.GetLeaveTangentWeight());
+							TangentScreenPosition.X += PixelDelta.X;
+							TangentScreenPosition.Y += PixelDelta.Y;
+							CurveEditor->GetTangentAndWeightFromScreenPosition(KeyScreenPosition, TangentScreenPosition, Tangent, Weight);
+								
+							FKeyAttributes NewAttributes;
+							NewAttributes.SetLeaveTangent(Tangent);
+							NewAttributes.SetLeaveTangentWeight(Weight);
+							NewKeyAttributesScratch.Add(NewAttributes);
+						}
+						else
+						{
+							const float PixelLength = 60.0f;
+							FVector2D TangentScreenPosition = FCurveEditor::GetVectorFromSlopeAndLength(LeaveTangent * -DisplayRatio, PixelLength) + KeyScreenPosition;
+							TangentScreenPosition += PixelDelta;
+							TangentScreenPosition -= KeyScreenPosition;
+							const float Tangent = (-TangentScreenPosition.Y / TangentScreenPosition.X) / DisplayRatio;
+							FKeyAttributes NewAttributes;
+							NewAttributes.SetLeaveTangent(Tangent);
+							NewKeyAttributesScratch.Add(NewAttributes);
+						}
+					}
+					else //still need to add since expect attributes to equal num of selected
+					{
+						NewKeyAttributesScratch.Add(FKeyAttributes());
 					}
 				}
 			}
-
 			Curve->SetKeyAttributes(KeyData.Handles, NewKeyAttributesScratch);
 		}
 	}
diff --git a/Engine/Source/Editor/CurveEditor/Private/DragOperations/CurveEditorDragOperation_Tangent.h b/Engine/Source/Editor/CurveEditor/Private/DragOperations/CurveEditorDragOperation_Tangent.h
--- a/Engine/Source/Editor/CurveEditor/Private/DragOperations/CurveEditorDragOperation_Tangent.h
+++ b/Engine/Source/Editor/CurveEditor/Private/DragOperations/CurveEditorDragOperation_Tangent.h
@@ -24,9 +24,6 @@ private:
 
 private:
 
-	/** Cardinal screen position to compute angles from (either a key position, or the starting mouse pixel) */
-	TOptional<FVector2D> CardinalPosition;
-
 	/** The type of tangents we're dragging */
 	ECurvePointType PointType;
 
diff --git a/Engine/Source/Editor/CurveEditor/Private/SCurveEditorPanel.cpp b/Engine/Source/Editor/CurveEditor/Private/SCurveEditorPanel.cpp
--- a/Engine/Source/Editor/CurveEditor/Private/SCurveEditorPanel.cpp
+++ b/Engine/Source/Editor/CurveEditor/Private/SCurveEditorPanel.cpp
@@ -24,6 +24,11 @@
 #include "Algo/Sort.h"
 #include "Widgets/Input/SNumericEntryBox.h"
 #include "Widgets/Input/STextEntryPopup.h"
+#include "CurveEditorEditObjectContainer.h"
+#include "Modules/ModuleManager.h"
+#include "PropertyEditorModule.h"
+#include "IDetailsView.h"
+
 #define LOCTEXT_NAMESPACE "SCurveEditorPanel"
 
 namespace
@@ -32,6 +37,11 @@ namespace
 	static float HoveredCurveThickness = 5.f;
 	static float UnHoveredCurveThickness = 1.f;
 	static bool  bAntiAliasCurves = true;
+
+	FReply HandledReply(const FGeometry&, const FPointerEvent&)
+	{
+		return FReply::Handled();
+	}
 }
 
 
@@ -55,6 +65,16 @@ public:
 	virtual bool IsEmpty() const override { return !bIsEnabled.Get(); }
 };
 
+SCurveEditorPanel::SCurveEditorPanel()
+{
+	EditObjects = MakeUnique<FCurveEditorEditObjectContainer>();
+	bSelectionSupportsWeightedTangents = false;
+}
+
+SCurveEditorPanel::~SCurveEditorPanel()
+{
+}
+
 void SCurveEditorPanel::Construct(const FArguments& InArgs, TSharedRef<FCurveEditor> InCurveEditor)
 {
 	GridLineTintAttribute = InArgs._GridLineTint;
@@ -72,11 +92,29 @@ void SCurveEditorPanel::Construct(const FArguments& InArgs, TSharedRef<FCurveEdi
 	BindCommands();
 	SetClipping(EWidgetClipping::ClipToBounds);
 
+	{
+		FPropertyEditorModule& EditModule = FModuleManager::Get().GetModuleChecked<FPropertyEditorModule>("PropertyEditor");
+		FDetailsViewArgs DetailsViewArgs(
+			/*bUpdateFromSelection=*/ false,
+			/*bLockable=*/ false,
+			/*bAllowSearch=*/ false,
+			FDetailsViewArgs::HideNameArea,
+			/*bHideSelectionTip=*/ true,
+			/*InNotifyHook=*/ nullptr,
+			/*InSearchInitialKeyFocus=*/ false,
+			/*InViewIdentifier=*/ NAME_None);
+		DetailsViewArgs.DefaultsOnlyVisibility = EEditDefaultsOnlyNodeVisibility::Automatic;
+		DetailsViewArgs.bShowOptions = false;
+
+		KeyDetailsView = EditModule.CreateDetailView(DetailsViewArgs);
+	}
+
 	ChildSlot
 	[
 		SNew(SOverlay)
 
 		+ SOverlay::Slot()
+		.VAlign(VAlign_Top)
 		[
 			SNew(SSplitter)
 			.Visibility(this, &SCurveEditorPanel::GetSplitterVisibility)
@@ -84,7 +122,15 @@ void SCurveEditorPanel::Construct(const FArguments& InArgs, TSharedRef<FCurveEdi
 			+ SSplitter::Slot()
 			.Value(0.25f)
 			[
-				SAssignNew(EditPanel, SScrollBox)
+				SNew(SBorder)
+				.Padding(FMargin(0))
+				.BorderImage(FEditorStyle::GetBrush("NoBorder"))
+				.OnMouseButtonDown_Static(HandledReply)
+				.OnMouseMove_Static(HandledReply)
+				.OnMouseButtonUp_Static(HandledReply)
+				[
+					KeyDetailsView.ToSharedRef()
+				]
 			]
 
 			+ SSplitter::Slot()
@@ -139,17 +185,23 @@ void SCurveEditorPanel::BindCommands()
 		FExecuteAction SetCubicUser  = FExecuteAction::CreateSP(this, &SCurveEditorPanel::SetKeyAttributes, FKeyAttributes().SetInterpMode(RCIM_Cubic).SetTangentMode(RCTM_User),    LOCTEXT("SetInterpUser",     "Set Interp User"));
 		FExecuteAction SetCubicBreak = FExecuteAction::CreateSP(this, &SCurveEditorPanel::SetKeyAttributes, FKeyAttributes().SetInterpMode(RCIM_Cubic).SetTangentMode(RCTM_Break),   LOCTEXT("SetInterpBreak",    "Set Interp Break"));
 
+		FExecuteAction    ToggleWeighted    = FExecuteAction::CreateSP(this, &SCurveEditorPanel::ToggleWeightedTangents);
+		FCanExecuteAction CanToggleWeighted = FCanExecuteAction::CreateSP(this, &SCurveEditorPanel::CanToggleWeightedTangents);
+
 		FIsActionChecked IsConstantCommon   = FIsActionChecked::CreateSP(this, &SCurveEditorPanel::CompareCommonInterpolationMode, RCIM_Constant);
 		FIsActionChecked IsLinearCommon     = FIsActionChecked::CreateSP(this, &SCurveEditorPanel::CompareCommonInterpolationMode, RCIM_Linear);
 		FIsActionChecked IsCubicAutoCommon  = FIsActionChecked::CreateSP(this, &SCurveEditorPanel::CompareCommonTangentMode, RCIM_Cubic, RCTM_Auto);
 		FIsActionChecked IsCubicUserCommon  = FIsActionChecked::CreateSP(this, &SCurveEditorPanel::CompareCommonTangentMode, RCIM_Cubic, RCTM_User);
 		FIsActionChecked IsCubicBreakCommon = FIsActionChecked::CreateSP(this, &SCurveEditorPanel::CompareCommonTangentMode, RCIM_Cubic, RCTM_Break);
+		FIsActionChecked IsCubicWeightCommon = FIsActionChecked::CreateSP(this, &SCurveEditorPanel::CompareCommonTangentWeightMode, RCIM_Cubic, RCTWM_WeightedBoth);
 
 		CommandList->MapAction(FCurveEditorCommands::Get().InterpolationConstant, SetConstant, FCanExecuteAction(), IsConstantCommon);
 		CommandList->MapAction(FCurveEditorCommands::Get().InterpolationLinear, SetLinear, FCanExecuteAction(), IsLinearCommon);
 		CommandList->MapAction(FCurveEditorCommands::Get().InterpolationCubicAuto, SetCubicAuto, FCanExecuteAction(), IsCubicAutoCommon);
 		CommandList->MapAction(FCurveEditorCommands::Get().InterpolationCubicUser, SetCubicUser, FCanExecuteAction(), IsCubicUserCommon);
 		CommandList->MapAction(FCurveEditorCommands::Get().InterpolationCubicBreak, SetCubicBreak, FCanExecuteAction(), IsCubicBreakCommon);
+		CommandList->MapAction(FCurveEditorCommands::Get().InterpolationToggleWeighted, ToggleWeighted, CanToggleWeighted, IsCubicWeightCommon);
+
 	}
 
 	// Pre Extrapolation Modes
@@ -208,6 +260,8 @@ void SCurveEditorPanel::Tick(const FGeometry& AllottedGeometry, const double InC
 
 	CachedDrawParams.Reset();
 	CurveEditor->GetCurveDrawParams(CachedDrawParams);
+
+	CachedSelectionSerialNumber = CurveEditor->Selection.GetSerialNumber();
 }
 
 void SCurveEditorPanel::UpdateCommonCurveInfo()
@@ -230,6 +284,7 @@ void SCurveEditorPanel::UpdateCommonCurveInfo()
 	}
 
 	// Reset the common curve and key info
+	bSelectionSupportsWeightedTangents = false;
 	CachedCommonCurveAttributes = AccumulatedCurveAttributes.Get(FCurveAttributes());
 
 	TOptional<FKeyAttributes> AccumulatedKeyAttributes;
@@ -246,6 +301,11 @@ void SCurveEditorPanel::UpdateCommonCurveInfo()
 			Curve->GetKeyAttributes(Pair.Value.AsArray(), AllKeyAttributes);
 			for (const FKeyAttributes& Attributes : AllKeyAttributes)
 			{
+				if (Attributes.HasTangentWeightMode())
+				{
+					bSelectionSupportsWeightedTangents = true;
+				}
+
 				if (!AccumulatedKeyAttributes.IsSet())
 				{
 					AccumulatedKeyAttributes = Attributes;
@@ -279,36 +339,70 @@ void SCurveEditorPanel::RebindContextualActions(FVector2D MousePosition)
 
 void SCurveEditorPanel::UpdateEditBox()
 {
-	const TMap<FCurveModelID, TUniquePtr<FCurveModel>>& AllCurves = CurveEditor->GetCurves();
-	for (const TTuple<FCurveModelID, TUniquePtr<FCurveModel>>& Pair : AllCurves)
+	const FCurveEditorSelection& Selection = CurveEditor->Selection;
+	for (TTuple<FCurveModelID, TMap<FKeyHandle, UObject*>>& OuterPair : EditObjects->CurveIDToKeyProxies)
 	{
-		if (!CurveToEditUI.Contains(Pair.Key))
+		const FKeyHandleSet* SelectedKeys = Selection.FindForCurve(OuterPair.Key);
+
+		for (TTuple<FKeyHandle, UObject*>& InnerPair : OuterPair.Value)
 		{
-			TSharedPtr<SWidget> EditUI = Pair.Value->CreateEditUI(CurveEditor, Pair.Key);
-			if (EditUI.IsValid())
+			if (ICurveEditorKeyProxy* Proxy = Cast<ICurveEditorKeyProxy>(InnerPair.Value))
 			{
-				EditPanel->AddSlot()
-				.Padding(FMargin(0,0,0,10))
-				[
-					EditUI.ToSharedRef()
-				];
+				Proxy->UpdateValuesFromRawData();
 			}
-			CurveToEditUI.Add(Pair.Key, EditUI);
 		}
 	}
 
-	for (auto It = CurveToEditUI.CreateIterator(); It; ++It)
+	if (CachedSelectionSerialNumber == Selection.GetSerialNumber())
+	{
+		return;
+	}
+
+	TArray<FKeyHandle> KeyHandleScratch;
+	TArray<UObject*>   NewProxiesScratch;
+
+	TArray<UObject*> AllEditObjects;
+	for (const TTuple<FCurveModelID, FKeyHandleSet>& Pair : Selection.GetAll())
 	{
-		if (!AllCurves.Contains(It.Key()))
+		FCurveModel* Curve = CurveEditor->FindCurve(Pair.Key);
+		if (!Curve)
 		{
-			if (It.Value().IsValid())
+			continue;
+		}
+
+		KeyHandleScratch.Reset();
+		NewProxiesScratch.Reset();
+
+		TMap<FKeyHandle, UObject*>& KeyHandleToEditObject = EditObjects->CurveIDToKeyProxies.FindOrAdd(Pair.Key);
+		for (FKeyHandle Handle : Pair.Value.AsArray())
+		{
+			if (UObject* Existing = KeyHandleToEditObject.FindRef(Handle))
 			{
-				EditPanel->RemoveSlot(It.Value().ToSharedRef());
+				AllEditObjects.Add(Existing);
 			}
+			else
+			{
+				KeyHandleScratch.Add(Handle);
+			}
+		}
+
+		if (KeyHandleScratch.Num() > 0)
+		{
+			NewProxiesScratch.SetNum(KeyHandleScratch.Num());
+			Curve->CreateKeyProxies(KeyHandleScratch, NewProxiesScratch);
 
-			It.RemoveCurrent();
+			for (int32 Index = 0; Index < KeyHandleScratch.Num(); ++Index)
+			{
+				if (UObject* NewObject = NewProxiesScratch[Index])
+				{
+					KeyHandleToEditObject.Add(KeyHandleScratch[Index], NewObject);
+					AllEditObjects.Add(NewObject);
+				}
+			}
 		}
 	}
+
+	KeyDetailsView->SetObjects(AllEditObjects);
 }
 
 void SCurveEditorPanel::UpdateCurveProximities(FVector2D MousePixel)
@@ -429,6 +523,65 @@ void SCurveEditorPanel::SetCurveAttributes(FCurveAttributes CurveAttributes, FTe
 	}
 }
 
+void SCurveEditorPanel::ToggleWeightedTangents()
+{
+	FScopedTransaction Transaction(LOCTEXT("ToggleWeightedTangents_Transaction", "Toggle Weighted Tangents"));
+
+	TMap<FCurveModelID, TArray<FKeyAttributes>> KeyAttributesPerCurve;
+
+	const TMap<FCurveModelID, FKeyHandleSet>& Selection = CurveEditor->Selection.GetAll();
+
+	// Disable weights unless we find something that doesn't have weights, then add them
+	FKeyAttributes KeyAttributesToAssign = FKeyAttributes().SetTangentWeightMode(RCTWM_WeightedNone);
+
+	// Gather current key attributes
+	for (const TTuple<FCurveModelID, FKeyHandleSet>& Pair : Selection)
+	{
+		FCurveModel* Curve = CurveEditor->FindCurve(Pair.Key);
+		if (Curve)
+		{
+			TArray<FKeyAttributes>& KeyAttributes = KeyAttributesPerCurve.Add(Pair.Key);
+			KeyAttributes.SetNum(Pair.Value.Num());
+			Curve->GetKeyAttributes(Pair.Value.AsArray(), KeyAttributes);
+
+			// Check all the key attributes if they support tangent weights, but don't have any. If we find any such keys, we'll enable weights on all.
+			if (KeyAttributesToAssign.GetTangentWeightMode() == RCTWM_WeightedNone)
+			{
+				for (const FKeyAttributes& Attributes : KeyAttributes)
+				{
+					if (Attributes.HasTangentWeightMode() && !(Attributes.HasArriveTangentWeight() || Attributes.HasLeaveTangentWeight()))
+					{
+						KeyAttributesToAssign.SetTangentWeightMode(RCTWM_WeightedBoth);
+						break;
+					}
+				}
+			}
+		}
+	}
+
+	// Assign the new key attributes to all the selected curves
+	for (TTuple<FCurveModelID, TArray<FKeyAttributes>>& Pair : KeyAttributesPerCurve)
+	{
+		FCurveModel* Curve = CurveEditor->FindCurve(Pair.Key);
+		if (Curve)
+		{
+			for (FKeyAttributes& Attributes : Pair.Value)
+			{
+				Attributes = KeyAttributesToAssign;
+			}
+
+			TArrayView<const FKeyHandle> KeyHandles = Selection.FindChecked(Pair.Key).AsArray();
+			Curve->Modify();
+			Curve->SetKeyAttributes(KeyHandles, Pair.Value);
+		}
+	}
+}
+
+bool SCurveEditorPanel::CanToggleWeightedTangents() const
+{
+	return bSelectionSupportsWeightedTangents;
+}
+
 void SCurveEditorPanel::OnAddKey(FVector2D MousePixel)
 {
 	FScopedTransaction Transcation(LOCTEXT("OnAddKey", "Add Key"));
@@ -834,7 +987,7 @@ int32 SCurveEditorPanel::OnPaint(const FPaintArgs& Args, const FGeometry& Allott
 	{
 		LayerId = DragOperation->DragImpl->Paint(AllottedGeometry, OutDrawElements, LayerId);
 	}
-	return SCompoundWidget::OnPaint(Args, AllottedGeometry, MyCullingRect, OutDrawElements, LayerId, InWidgetStyle, bParentEnabled);
+	return SCompoundWidget::OnPaint(Args, AllottedGeometry, MyCullingRect, OutDrawElements, LayerId+2000, InWidgetStyle, bParentEnabled)-2000;
 }
 
 int32 SCurveEditorPanel::DrawGridLines(const FGeometry& AllottedGeometry, FSlateWindowElementList& OutDrawElements, int32 LayerId, ESlateDrawEffect DrawEffects) const
