Clone of CL 4059587:  Fix for Sequencer window opening behind Main Window.  FLevelSequenceEditorToolkit::Initialize() Closes and reopens the world outliner after adding the Sequencer window.  On Mac this makes the world outliner / main window move back in front of the sequencer window.  This change swaps around those operations.  Reopen the world outliner then add the sequencer window.  Outliner info still changes it's visible data as before.

Tested with World Outliner docked and undocked.

diff --git a/Engine/Plugins/MovieScene/LevelSequenceEditor/Source/LevelSequenceEditor/Private/LevelSequenceEditorToolkit.cpp b/Engine/Plugins/MovieScene/LevelSequenceEditor/Source/LevelSequenceEditor/Private/LevelSequenceEditorToolkit.cpp
--- a/Engine/Plugins/MovieScene/LevelSequenceEditor/Source/LevelSequenceEditor/Private/LevelSequenceEditorToolkit.cpp
+++ b/Engine/Plugins/MovieScene/LevelSequenceEditor/Source/LevelSequenceEditor/Private/LevelSequenceEditorToolkit.cpp
@@ -264,15 +264,16 @@ void FLevelSequenceEditorToolkit::Initialize(const EToolkitMode::Type Mode, cons
 	// @todo remove when world-centric mode is added
 	FLevelEditorModule& LevelEditorModule = FModuleManager::LoadModuleChecked<FLevelEditorModule>("LevelEditor");
 
-	LevelEditorModule.AttachSequencer(Sequencer->GetSequencerWidget(), SharedThis(this));
-
-	// @todo reopen the scene outliner so that is refreshed with the sequencer info column
+	// Reopen the scene outliner so that is refreshed with the sequencer info column
 	TSharedPtr<FTabManager> LevelEditorTabManager = LevelEditorModule.GetLevelEditorTabManager();
 	if (LevelEditorTabManager->FindExistingLiveTab(FName("LevelEditorSceneOutliner")).IsValid())
 	{
 		LevelEditorTabManager->InvokeTab(FName("LevelEditorSceneOutliner"))->RequestCloseTab();
 		LevelEditorTabManager->InvokeTab(FName("LevelEditorSceneOutliner"));
 	}
+	
+	// Now Attach so this window will apear in the correct front first order
+	LevelEditorModule.AttachSequencer(Sequencer->GetSequencerWidget(), SharedThis(this));
 
 	// We need to find out when the user loads a new map, because we might need to re-create puppet actors
 	// when previewing a MovieScene
