Sequencer: Refresh actor list to clean up "deleted actor" in the outliner when converting to spawnable.

diff --git a/Engine/Plugins/MovieScene/LevelSequenceEditor/Source/LevelSequenceEditor/Private/Misc/LevelSequenceEditorSpawnRegister.cpp b/Engine/Plugins/MovieScene/LevelSequenceEditor/Source/LevelSequenceEditor/Private/Misc/LevelSequenceEditorSpawnRegister.cpp
--- a/Engine/Plugins/MovieScene/LevelSequenceEditor/Source/LevelSequenceEditor/Private/Misc/LevelSequenceEditorSpawnRegister.cpp
+++ b/Engine/Plugins/MovieScene/LevelSequenceEditor/Source/LevelSequenceEditor/Private/Misc/LevelSequenceEditorSpawnRegister.cpp
@@ -347,6 +347,8 @@ void FLevelSequenceEditorSpawnRegister::HandleConvertPossessableToSpawnable(UObj
 		if (World)
 		{
 			World->EditorDestroyActor(OldActor, true);
+
+			GEditor->BroadcastLevelActorListChanged();
 		}
 	}
 }
