Opening the Add Track menu in Sequencer and trying to add an actor from the map to the sequence should take significantly less time now on maps that have tens of thousands of actors.

The old method (through a series of complicated events) cleared the binding cache in sequencer and re-made it for every binding in the sequence for every possible actor. The new method gets all bindings once, finds all actors they're bound to, and then checks each actor in the map against that set.

diff --git a/Engine/Plugins/MovieScene/LevelSequenceEditor/Source/LevelSequenceEditor/Private/Misc/LevelSequenceEditorActorBinding.cpp b/Engine/Plugins/MovieScene/LevelSequenceEditor/Source/LevelSequenceEditor/Private/Misc/LevelSequenceEditorActorBinding.cpp
--- a/Engine/Plugins/MovieScene/LevelSequenceEditor/Source/LevelSequenceEditor/Private/Misc/LevelSequenceEditorActorBinding.cpp
+++ b/Engine/Plugins/MovieScene/LevelSequenceEditor/Source/LevelSequenceEditor/Private/Misc/LevelSequenceEditorActorBinding.cpp
@@ -13,6 +13,7 @@
 #include "Modules/ModuleManager.h"
 #include "Editor.h"
 #include "Widgets/Layout/SBox.h"
+#include "MovieScene.h"
 
 #define LOCTEXT_NAMESPACE "LevelSequenceEditorActorBinding"
 
@@ -39,17 +40,38 @@ bool FLevelSequenceEditorActorBinding::SupportsSequence(UMovieSceneSequence* InS
 
 void FLevelSequenceEditorActorBinding::AddPossessActorMenuExtensions(FMenuBuilder& MenuBuilder)
 {
-	auto IsActorValidForPossession = [=](const AActor* InActor, TWeakPtr<ISequencer> InWeakSequencer)
+	// This is called for every actor in the map, and asking the sequencer for a handle to the object to check if we have
+	// already bound is an issue on maps that have tens of thousands of actors. The current sequence is will almost always
+	// have actors than the map, so instead we'll cache off all of the actors already bound and check against that map locally.
+	// This list is checked via an async filter, but we don't need to store them as weak pointers because we're doing a direct
+	// pointer comparison and not an object comparison, and the async list shouldn't run the filter if the object is no longer valid.
+	// We don't need to check against Sequencer spawnables as they're not valid for possession.
+	TSet<UObject*> ExistingPossessedObjects;
+	if (Sequencer.IsValid())
 	{
-		TSharedPtr<ISequencer> SequencerPtr = InWeakSequencer.Pin();
-		bool bCreateHandleIfMissing = false;
-		return SequencerPtr.IsValid() && !SequencerPtr->GetHandleToObject((UObject*)InActor, bCreateHandleIfMissing).IsValid();
+		UMovieSceneSequence* MovieSceneSequence = Sequencer.Pin()->GetFocusedMovieSceneSequence();
+		UMovieScene* MovieScene = MovieSceneSequence->GetMovieScene();
+		if(MovieScene)
+		{
+			for (int32 Index = 0; Index < MovieScene->GetPossessableCount(); Index++)
+			{
+				FMovieScenePossessable& Possessable = MovieScene->GetPossessable(Index);
+
+				// A possession guid can apply to more than one object, so we get all bound objects for the GUID and add them to our set.
+				ExistingPossessedObjects.Append(MovieSceneSequence->LocateBoundObjects(Possessable.GetGuid(), Sequencer.Pin()->GetPlaybackContext()));
+			}
+		}
+	}
+
+	auto IsActorValidForPossession = [=](const AActor* InActor, TSet<UObject*> InPossessedObjectSet)
+	{
+		return !InPossessedObjectSet.Contains((UObject*)InActor);
 	};
 
 	// Set up a menu entry to add the selected actor(s) to the sequencer
 	TArray<AActor*> ActorsValidForPossession;
 	GEditor->GetSelectedActors()->GetSelectedObjects(ActorsValidForPossession);
-	ActorsValidForPossession.RemoveAll([&](AActor* In){ return !IsActorValidForPossession(In, Sequencer); });
+	ActorsValidForPossession.RemoveAll([&](AActor* In){ return !IsActorValidForPossession(In, ExistingPossessedObjects); });
 
 	FText SelectedLabel;
 	FSlateIcon ActorIcon = FSlateIconFinder::FindIconForClass(AActor::StaticClass());
@@ -91,7 +113,7 @@ void FLevelSequenceEditorActorBinding::AddPossessActorMenuExtensions(FMenuBuilde
 		InitOptions.ColumnMap.Add(FBuiltInColumnTypes::Label(), FColumnInfo(EColumnVisibility::Visible, 0));
 
 		// Only display actors that are not possessed already
-		InitOptions.Filters->AddFilterPredicate(FActorFilterPredicate::CreateLambda(IsActorValidForPossession, Sequencer));
+		InitOptions.Filters->AddFilterPredicate(FActorFilterPredicate::CreateLambda(IsActorValidForPossession, ExistingPossessedObjects));
 	}
 
 	// actor selector to allow the user to choose an actor
