Added tooltips to the Cinematic Viewport UI to help clarify what some of the numbers being shown are.

diff --git a/Engine/Plugins/MovieScene/LevelSequenceEditor/Source/LevelSequenceEditor/Private/CinematicViewport/SCinematicLevelViewport.cpp b/Engine/Plugins/MovieScene/LevelSequenceEditor/Source/LevelSequenceEditor/Private/CinematicViewport/SCinematicLevelViewport.cpp
--- a/Engine/Plugins/MovieScene/LevelSequenceEditor/Source/LevelSequenceEditor/Private/CinematicViewport/SCinematicLevelViewport.cpp
+++ b/Engine/Plugins/MovieScene/LevelSequenceEditor/Source/LevelSequenceEditor/Private/CinematicViewport/SCinematicLevelViewport.cpp
@@ -205,6 +205,7 @@ void SCinematicLevelViewport::Construct(const FArguments& InArgs)
 				.MaxValue(TOptional<double>())
 				.OnEndSliderMovement(this, &SCinematicLevelViewport::SetTime)
 				.Value(this, &SCinematicLevelViewport::GetTime)
+				.ToolTipText(LOCTEXT("TimeLocalToCurrentSequence", "The current time of the sequence relative to the focused sequence."))
 				.Delta_Lambda([=]()
 				{
 					return UIData.OuterResolution.AsDecimal() * UIData.OuterPlayRate.AsInterval(); 
@@ -309,6 +310,7 @@ void SCinematicLevelViewport::Construct(const FArguments& InArgs)
 									SNew(STextBlock)
 									.ColorAndOpacity(Gray)
 									.Text_Lambda([=]{ return UIData.ShotName; })
+									.ToolTipText(LOCTEXT("CurrentSequence", "The name of the currently evaluated sequence."))
 								]
 
 								+ SHorizontalBox::Slot()
@@ -318,7 +320,8 @@ void SCinematicLevelViewport::Construct(const FArguments& InArgs)
 								[
 									SNew(STextBlock)
 									.ColorAndOpacity(Gray)
-									.Text_Lambda([=]{ return UIData.CameraName; })
+									.Text_Lambda([=] { return UIData.CameraName; })
+									.ToolTipText(LOCTEXT("CurrentCamera", "The name of the current camera."))
 								]
 							]
 
@@ -328,7 +331,8 @@ void SCinematicLevelViewport::Construct(const FArguments& InArgs)
 							[
 								SNew(STextBlock)
 								.ColorAndOpacity(Gray)
-								.Text_Lambda([=]{ return UIData.Filmback; })
+								.Text_Lambda([=] { return UIData.Filmback; })
+								.ToolTipText(LOCTEXT("CurrentFilmback", "The name of the current shot's filmback (the imaging area of the frame/sensor)."))
 							]
 
 							+ SHorizontalBox::Slot()
@@ -337,7 +341,8 @@ void SCinematicLevelViewport::Construct(const FArguments& InArgs)
 								SNew(STextBlock)
 								.Font(FEditorStyle::GetFontStyle("Sequencer.FixedFont"))
 								.ColorAndOpacity(Gray)
-								.Text_Lambda([=]{ return UIData.LocalPlaybackTime; })
+								.Text_Lambda([=] { return UIData.LocalPlaybackTime; })
+								.ToolTipText(LOCTEXT("LocalPlaybackTime", "The current playback time relative to the currently evaluated sequence."))
 							]
 						]
 					]
diff --git a/Engine/Source/Editor/SequencerWidgets/Private/STimeRange.cpp b/Engine/Source/Editor/SequencerWidgets/Private/STimeRange.cpp
--- a/Engine/Source/Editor/SequencerWidgets/Private/STimeRange.cpp
+++ b/Engine/Source/Editor/SequencerWidgets/Private/STimeRange.cpp
@@ -81,7 +81,7 @@ void STimeRange::Construct( const STimeRange::FArguments& InArgs, TSharedRef<ITi
 	{
 		PlaybackRangeStart = SNew(SSpinBox<double>)
 		.Value(this, &STimeRange::PlayStartTime)
-		.ToolTipText(LOCTEXT("PlayStartTimeTooltip", "In Time"))
+		.ToolTipText(LOCTEXT("PlayStartTimeTooltip", "Playback Range Start Time"))
 		.OnValueCommitted(this, &STimeRange::OnPlayStartTimeCommitted)
 		.OnValueChanged(this, &STimeRange::OnPlayStartTimeChanged)
 		.MinValue(TOptional<double>())
@@ -95,7 +95,7 @@ void STimeRange::Construct( const STimeRange::FArguments& InArgs, TSharedRef<ITi
 
 		PlaybackRangeEnd = SNew(SSpinBox<double>)
 		.Value(this, &STimeRange::PlayEndTime)
-		.ToolTipText(LOCTEXT("PlayEndTimeTooltip", "Out Time"))
+		.ToolTipText(LOCTEXT("PlayEndTimeTooltip", "Playback Range Stop Time"))
 		.OnValueCommitted( this, &STimeRange::OnPlayEndTimeCommitted )
 		.OnValueChanged( this, &STimeRange::OnPlayEndTimeChanged )
 		.MinValue(TOptional<double>())
