Sequencer: Add default sections for all tracks (not just property tracks)

diff --git a/Engine/Plugins/MovieScene/LevelSequenceEditor/Source/LevelSequenceEditor/Private/LevelSequenceEditorToolkit.cpp b/Engine/Plugins/MovieScene/LevelSequenceEditor/Source/LevelSequenceEditor/Private/LevelSequenceEditorToolkit.cpp
--- a/Engine/Plugins/MovieScene/LevelSequenceEditor/Source/LevelSequenceEditor/Private/LevelSequenceEditorToolkit.cpp
+++ b/Engine/Plugins/MovieScene/LevelSequenceEditor/Source/LevelSequenceEditor/Private/LevelSequenceEditorToolkit.cpp
@@ -422,56 +422,52 @@ void FLevelSequenceEditorToolkit::AddDefaultTracksForActor(AActor& Actor, const
 					NewTrack = MovieScene->AddTrack(TrackClass, Binding);
 				}
 
-				// Create a section for any property tracks
-				if (Cast<UMovieScenePropertyTrack>(NewTrack))
+				UMovieSceneSection* NewSection;
+				if (NewTrack->GetAllSections().Num() > 0)
 				{
-					UMovieSceneSection* NewSection;
-					if (NewTrack->GetAllSections().Num() > 0)
-					{
-						NewSection = NewTrack->GetAllSections()[0];
-					}
-					else
-					{
-						NewSection = NewTrack->CreateNewSection();
-						NewTrack->AddSection(*NewSection);
-					}
+					NewSection = NewTrack->GetAllSections()[0];
+				}
+				else
+				{
+					NewSection = NewTrack->CreateNewSection();
+					NewTrack->AddSection(*NewSection);
+				}
 
-					// @todo sequencer: hack: setting defaults for transform tracks
-					if (NewTrack->IsA(UMovieScene3DTransformTrack::StaticClass()) && Sequencer->GetAutoSetTrackDefaults())
-					{
-						auto TransformSection = Cast<UMovieScene3DTransformSection>(NewSection);
+				// @todo sequencer: hack: setting defaults for transform tracks
+				if (NewTrack->IsA(UMovieScene3DTransformTrack::StaticClass()) && Sequencer->GetAutoSetTrackDefaults())
+				{
+					auto TransformSection = Cast<UMovieScene3DTransformSection>(NewSection);
 
-						FVector Location = Actor.GetActorLocation();
-						FRotator Rotation = Actor.GetActorRotation();
-						FVector Scale = Actor.GetActorScale();
+					FVector Location = Actor.GetActorLocation();
+					FRotator Rotation = Actor.GetActorRotation();
+					FVector Scale = Actor.GetActorScale();
 
-						if (Actor.GetRootComponent())
-						{
-							FTransform ActorRelativeTransform = Actor.GetRootComponent()->GetRelativeTransform();
+					if (Actor.GetRootComponent())
+					{
+						FTransform ActorRelativeTransform = Actor.GetRootComponent()->GetRelativeTransform();
 
-							Location = ActorRelativeTransform.GetTranslation();
-							Rotation = ActorRelativeTransform.GetRotation().Rotator();
-							Scale = ActorRelativeTransform.GetScale3D();
-						}
+						Location = ActorRelativeTransform.GetTranslation();
+						Rotation = ActorRelativeTransform.GetRotation().Rotator();
+						Scale = ActorRelativeTransform.GetScale3D();
+					}
 
-						TArrayView<FMovieSceneFloatChannel*> FloatChannels = TransformSection->GetChannelProxy().GetChannels<FMovieSceneFloatChannel>();
-						FloatChannels[0]->SetDefault(Location.X);
-						FloatChannels[1]->SetDefault(Location.Y);
-						FloatChannels[2]->SetDefault(Location.Z);
+					TArrayView<FMovieSceneFloatChannel*> FloatChannels = TransformSection->GetChannelProxy().GetChannels<FMovieSceneFloatChannel>();
+					FloatChannels[0]->SetDefault(Location.X);
+					FloatChannels[1]->SetDefault(Location.Y);
+					FloatChannels[2]->SetDefault(Location.Z);
 
-						FloatChannels[3]->SetDefault(Rotation.Euler().X);
-						FloatChannels[4]->SetDefault(Rotation.Euler().Y);
-						FloatChannels[5]->SetDefault(Rotation.Euler().Z);
+					FloatChannels[3]->SetDefault(Rotation.Euler().X);
+					FloatChannels[4]->SetDefault(Rotation.Euler().Y);
+					FloatChannels[5]->SetDefault(Rotation.Euler().Z);
 
-						FloatChannels[6]->SetDefault(Scale.X);
-						FloatChannels[7]->SetDefault(Scale.Y);
-						FloatChannels[8]->SetDefault(Scale.Z);
-					}
+					FloatChannels[6]->SetDefault(Scale.X);
+					FloatChannels[7]->SetDefault(Scale.Y);
+					FloatChannels[8]->SetDefault(Scale.Z);
+				}
 
-					if (GetSequencer()->GetInfiniteKeyAreas())
-					{
-						NewSection->SetRange(TRange<FFrameNumber>::All());
-					}
+				if (GetSequencer()->GetInfiniteKeyAreas())
+				{
+					NewSection->SetRange(TRange<FFrameNumber>::All());
 				}
 			}
 		}
