Added missing call to FKismetCompilerContext::RegisterCompilerForBP

diff --git a/Engine/Plugins/MovieScene/LevelSequenceEditor/Source/LevelSequenceEditor/Private/LevelSequenceEditorModule.cpp b/Engine/Plugins/MovieScene/LevelSequenceEditor/Source/LevelSequenceEditor/Private/LevelSequenceEditorModule.cpp
--- a/Engine/Plugins/MovieScene/LevelSequenceEditor/Source/LevelSequenceEditor/Private/LevelSequenceEditorModule.cpp
+++ b/Engine/Plugins/MovieScene/LevelSequenceEditor/Source/LevelSequenceEditor/Private/LevelSequenceEditorModule.cpp
@@ -37,6 +37,7 @@
 #include "SequencerSettings.h"
 #include "Misc/MovieSceneSequenceEditor_LevelSequence.h"
 #include "KismetCompilerModule.h"
+#include "KismetCompiler.h"
 #include "BlueprintAssetHandler.h"
 
 #define LOCTEXT_NAMESPACE "LevelSequenceEditor"
@@ -76,6 +77,8 @@ public:
 		IKismetCompilerInterface& KismetCompilerModule = FModuleManager::LoadModuleChecked<IKismetCompilerInterface>("KismetCompiler");
 		KismetCompilerModule.GetCompilers().Add(this);
 
+		FKismetCompilerContext::RegisterCompilerForBP(ULevelSequenceDirectorBlueprint::StaticClass(), &FactoryCompilerContext);
+
 		class FLevelSequenceAssetBlueprintHandler : public IBlueprintAssetHandler
 		{
 			virtual UBlueprint* RetrieveBlueprint(UObject* InObject) const override
@@ -320,6 +323,11 @@ protected:
 
 protected:
 
+	static TSharedPtr<FKismetCompilerContext> FactoryCompilerContext(UBlueprint* BP, FCompilerResultsLog& InMessageLog, const FKismetCompilerOptions& InCompileOptions)
+	{
+		return MakeShared<FLevelSequenceDirectorBlueprintCompiler>(CastChecked<ULevelSequenceDirectorBlueprint>(BP), InMessageLog, InCompileOptions, nullptr);
+	}
+
 	virtual bool CanCompile(const UBlueprint* Blueprint) override
 	{
 		return Blueprint && Blueprint->IsA<ULevelSequenceDirectorBlueprint>();
