Sequencer: Level sequence editor blueprint library

diff --git a/Engine/Plugins/MovieScene/LevelSequenceEditor/Source/LevelSequenceEditor/Private/LevelSequenceEditorBlueprintLibrary.cpp b/Engine/Plugins/MovieScene/LevelSequenceEditor/Source/LevelSequenceEditor/Private/LevelSequenceEditorBlueprintLibrary.cpp
new file mode 100644
--- /dev/null
+++ b/Engine/Plugins/MovieScene/LevelSequenceEditor/Source/LevelSequenceEditor/Private/LevelSequenceEditorBlueprintLibrary.cpp
@@ -0,0 +1,90 @@
+// Copyright 1998-2019 Epic Games, Inc. All Rights Reserved.
+
+#include "LevelSequenceEditorBlueprintLibrary.h"
+
+#include "ISequencer.h"
+#include "LevelSequence.h"
+
+#include "Modules/ModuleManager.h"
+#include "LevelEditor.h"
+
+namespace
+{
+	static TWeakPtr<ISequencer> CurrentSequencer;
+}
+
+bool ULevelSequenceEditorBlueprintLibrary::OpenLevelSequence(ULevelSequence* LevelSequence)
+{
+	return FAssetEditorManager::Get().OpenEditorForAsset(LevelSequence);
+}
+
+void ULevelSequenceEditorBlueprintLibrary::CloseLevelSequence()
+{
+	if (CurrentSequencer.IsValid())
+	{
+		FAssetEditorManager::Get().CloseAllEditorsForAsset(CurrentSequencer.Pin()->GetRootMovieSceneSequence());
+	}
+}
+
+void ULevelSequenceEditorBlueprintLibrary::Play()
+{
+	const bool bTogglePlay = false;
+	if (CurrentSequencer.IsValid())
+	{
+		CurrentSequencer.Pin()->OnPlay(bTogglePlay);
+	}
+}
+
+void ULevelSequenceEditorBlueprintLibrary::Pause()
+{
+	if (CurrentSequencer.IsValid())
+	{
+		CurrentSequencer.Pin()->Pause();
+	}
+}
+
+void ULevelSequenceEditorBlueprintLibrary::SetCurrentTime(int32 NewFrame)
+{
+	if (CurrentSequencer.IsValid())
+	{
+		FFrameRate DisplayRate = CurrentSequencer.Pin()->GetFocusedDisplayRate();
+		FFrameRate TickResolution = CurrentSequencer.Pin()->GetFocusedTickResolution();
+
+		CurrentSequencer.Pin()->SetGlobalTime(ConvertFrameTime(NewFrame, DisplayRate, TickResolution));
+	}
+}
+
+int32 ULevelSequenceEditorBlueprintLibrary::GetCurrentTime()
+{
+	if (CurrentSequencer.IsValid())
+	{
+		FFrameRate DisplayRate = CurrentSequencer.Pin()->GetFocusedDisplayRate();
+		FFrameRate TickResolution = CurrentSequencer.Pin()->GetFocusedTickResolution();
+
+		return ConvertFrameTime(CurrentSequencer.Pin()->GetGlobalTime().Time, TickResolution, DisplayRate).FloorToFrame().Value;
+	}
+	return 0;
+}
+
+bool ULevelSequenceEditorBlueprintLibrary::IsPlaying()
+{
+	if (CurrentSequencer.IsValid())
+	{
+		return CurrentSequencer.Pin()->GetPlaybackStatus() == EMovieScenePlayerStatus::Playing;
+	}
+	return false;
+}
+
+bool ULevelSequenceEditorBlueprintLibrary::IsPaused()
+{
+	if (CurrentSequencer.IsValid())
+	{
+		return CurrentSequencer.Pin()->GetPlaybackStatus() == EMovieScenePlayerStatus::Paused;
+	}
+	return false;
+}
+
+void ULevelSequenceEditorBlueprintLibrary::SetSequencer(TSharedRef<ISequencer> InSequencer)
+{
+	CurrentSequencer = TWeakPtr<ISequencer>(InSequencer);
+}
\ No newline at end of file
diff --git a/Engine/Plugins/MovieScene/LevelSequenceEditor/Source/LevelSequenceEditor/Private/LevelSequenceEditorBlueprintLibrary.h b/Engine/Plugins/MovieScene/LevelSequenceEditor/Source/LevelSequenceEditor/Private/LevelSequenceEditorBlueprintLibrary.h
new file mode 100644
--- /dev/null
+++ b/Engine/Plugins/MovieScene/LevelSequenceEditor/Source/LevelSequenceEditor/Private/LevelSequenceEditorBlueprintLibrary.h
@@ -0,0 +1,80 @@
+// Copyright 1998-2019 Epic Games, Inc. All Rights Reserved.
+
+#pragma once
+
+#include "Misc/QualifiedFrameTime.h"
+#include "Kismet/BlueprintFunctionLibrary.h"
+#include "LevelSequenceEditorBlueprintLibrary.generated.h"
+
+class ISequencer;
+class ULevelSequence;
+
+UCLASS()
+class ULevelSequenceEditorBlueprintLibrary : public UBlueprintFunctionLibrary
+{
+public:
+
+	GENERATED_BODY()
+
+	/*
+	 * Open a level sequence asset
+	 */
+	UFUNCTION(BlueprintCallable, Category = "Level Sequence Editor")
+	static bool OpenLevelSequence(ULevelSequence* LevelSequence);
+
+	/*
+	 * Close
+	 */
+	UFUNCTION(BlueprintCallable, Category = "Level Sequence Editor")
+	static void CloseLevelSequence();
+
+	/**
+	 * Play the current level sequence
+	 */
+	UFUNCTION(BlueprintCallable, Category = "Level Sequence Editor")
+	static void Play();
+
+	/**
+	 * Pause the current level sequence
+	 */
+	UFUNCTION(BlueprintCallable, Category = "Level Sequence Editor")
+	static void Pause();
+
+public:
+
+	/**
+	 * Set playback position for the current level sequence in frames
+	 */
+	UFUNCTION(BlueprintCallable, Category = "Level Sequence Editor")
+	static void SetCurrentTime(int32 NewFrame);
+
+	/**
+	 * Get the current playback position in frames
+	 */
+	UFUNCTION(BlueprintCallable, Category = "Level Sequence Editor")
+	static int32 GetCurrentTime();
+
+public:
+
+	/** Check whether the sequence is actively playing. */
+	UFUNCTION(BlueprintPure, Category = "Level Sequence Editor")
+	static bool IsPlaying();
+
+	/** Check whether the sequence is paused. */
+	UFUNCTION(BlueprintPure, Category = "Level Sequence Editor")
+	static bool IsPaused();
+
+public:
+
+	/*
+	 * Callbacks
+	 */
+
+public:
+
+	 /**
+	  * Internal function to assign a sequencer singleton.
+	  * NOTE: Only to be called by LevelSequenceEditor::Construct.
+	  */
+	static void SetSequencer(TSharedRef<ISequencer> InSequencer);
+};
\ No newline at end of file
diff --git a/Engine/Plugins/MovieScene/LevelSequenceEditor/Source/LevelSequenceEditor/Private/LevelSequenceEditorToolkit.cpp b/Engine/Plugins/MovieScene/LevelSequenceEditor/Source/LevelSequenceEditor/Private/LevelSequenceEditorToolkit.cpp
--- a/Engine/Plugins/MovieScene/LevelSequenceEditor/Source/LevelSequenceEditor/Private/LevelSequenceEditorToolkit.cpp
+++ b/Engine/Plugins/MovieScene/LevelSequenceEditor/Source/LevelSequenceEditor/Private/LevelSequenceEditorToolkit.cpp
@@ -48,6 +48,7 @@
 #include "Widgets/Docking/SDockTab.h"
 #include "SequencerSettings.h"
 #include "LevelEditorSequencerIntegration.h"
+#include "LevelSequenceEditorBlueprintLibrary.h"
 #include "MovieSceneCaptureDialogModule.h"
 #include "MovieScene.h"
 #include "UnrealEdMisc.h"
@@ -198,6 +199,7 @@ void FLevelSequenceEditorToolkit::Initialize(const EToolkitMode::Type Mode, cons
 	Options.bCanRecord = true;
 
 	FLevelEditorSequencerIntegration::Get().AddSequencer(Sequencer.ToSharedRef(), Options);
+	ULevelSequenceEditorBlueprintLibrary::SetSequencer(Sequencer.ToSharedRef());
 
 	// @todo remove when world-centric mode is added
 	FLevelEditorModule& LevelEditorModule = FModuleManager::LoadModuleChecked<FLevelEditorModule>("LevelEditor");
