Cinematic Viewports now correctly enable Game view mode when activated.

diff --git a/Engine/Plugins/MovieScene/LevelSequenceEditor/Source/LevelSequenceEditor/Private/CinematicViewport/SCinematicLevelViewport.cpp b/Engine/Plugins/MovieScene/LevelSequenceEditor/Source/LevelSequenceEditor/Private/CinematicViewport/SCinematicLevelViewport.cpp
--- a/Engine/Plugins/MovieScene/LevelSequenceEditor/Source/LevelSequenceEditor/Private/CinematicViewport/SCinematicLevelViewport.cpp
+++ b/Engine/Plugins/MovieScene/LevelSequenceEditor/Source/LevelSequenceEditor/Private/CinematicViewport/SCinematicLevelViewport.cpp
@@ -84,7 +84,6 @@ FCinematicViewportClient::FCinematicViewportClient()
 {
 	bDrawAxes = false;
 	bIsRealtime = true;
-	SetGameView(true);
 	SetAllowCinematicControl(true);
 	bDisableInput = false;
 }
@@ -169,6 +168,11 @@ void SCinematicLevelViewport::Construct(const FArguments& InArgs)
 		.Realtime(true);
 
 	ViewportClient->SetViewportWidget(ViewportWidget);
+	
+	// Automatically engage game-view to hide editor only sprites. This needs to be done
+	// after the Viewport Client and Widget are constructed as they reset the view to defaults
+	// as part of their initialization.
+	ViewportClient->SetGameView(true);
 
 	TypeInterfaceProxy = MakeShareable( new FTypeInterfaceProxy );
 
