Sequencer: Removed derived blueprint type for level sequence director blueprints

diff --git a/Engine/Plugins/MovieScene/LevelSequenceEditor/Source/LevelSequenceEditor/Private/LevelSequenceDirectorBlueprint.cpp b/Engine/Plugins/MovieScene/LevelSequenceEditor/Source/LevelSequenceEditor/Private/LevelSequenceDirectorBlueprint.cpp
deleted file mode 100644
diff --git a/Engine/Plugins/MovieScene/LevelSequenceEditor/Source/LevelSequenceEditor/Private/LevelSequenceDirectorBlueprint.h b/Engine/Plugins/MovieScene/LevelSequenceEditor/Source/LevelSequenceEditor/Private/LevelSequenceDirectorBlueprint.h
deleted file mode 100644
diff --git a/Engine/Plugins/MovieScene/LevelSequenceEditor/Source/LevelSequenceEditor/Private/LevelSequenceEditorModule.cpp b/Engine/Plugins/MovieScene/LevelSequenceEditor/Source/LevelSequenceEditor/Private/LevelSequenceEditorModule.cpp
--- a/Engine/Plugins/MovieScene/LevelSequenceEditor/Source/LevelSequenceEditor/Private/LevelSequenceEditorModule.cpp
+++ b/Engine/Plugins/MovieScene/LevelSequenceEditor/Source/LevelSequenceEditor/Private/LevelSequenceEditorModule.cpp
@@ -36,8 +36,6 @@
 #include "Misc/LevelSequenceEditorActorSpawner.h"
 #include "SequencerSettings.h"
 #include "Misc/MovieSceneSequenceEditor_LevelSequence.h"
-#include "KismetCompilerModule.h"
-#include "KismetCompiler.h"
 #include "BlueprintAssetHandler.h"
 
 #define LOCTEXT_NAMESPACE "LevelSequenceEditor"
@@ -50,7 +48,7 @@ TSharedPtr<FLevelSequenceEditorStyle> FLevelSequenceEditorStyle::Singleton;
  * Implements the LevelSequenceEditor module.
  */
 class FLevelSequenceEditorModule
-	: public ILevelSequenceEditorModule, public IBlueprintCompiler, public FGCObject
+	: public ILevelSequenceEditorModule, public FGCObject
 {
 public:
 
@@ -74,11 +72,6 @@ public:
 		RegisterSettings();
 		RegisterSequenceEditor();
 
-		IKismetCompilerInterface& KismetCompilerModule = FModuleManager::LoadModuleChecked<IKismetCompilerInterface>("KismetCompiler");
-		KismetCompilerModule.GetCompilers().Add(this);
-
-		FKismetCompilerContext::RegisterCompilerForBP(ULevelSequenceDirectorBlueprint::StaticClass(), &FactoryCompilerContext);
-
 		class FLevelSequenceAssetBlueprintHandler : public IBlueprintAssetHandler
 		{
 			virtual UBlueprint* RetrieveBlueprint(UObject* InObject) const override
@@ -91,6 +84,15 @@ public:
 				// Only have a blueprint if it contains the BlueprintPathWithinPackage tag
 				return InAssetData.TagsAndValues.Find(FBlueprintTags::BlueprintPathWithinPackage) != nullptr;
 			}
+
+			virtual bool SupportsNativization(const UObject* InAsset, const UBlueprint* InBlueprint, FText* OutReason) const
+			{
+				if (OutReason)
+				{
+					*OutReason = LOCTEXT("NativizationError", "Level Sequences do not support nativization.");
+				}
+				return false;
+			}
 		};
 
 		FBlueprintAssetHandler::Get().RegisterHandler<FLevelSequenceAssetBlueprintHandler>(ULevelSequence::StaticClass()->GetFName());
@@ -106,12 +108,6 @@ public:
 		UnregisterPlacementModeExtensions();
 		UnregisterSettings();
 		UnregisterSequenceEditor();
-
-		IKismetCompilerInterface* KismetCompilerModulePtr = FModuleManager::GetModulePtr<IKismetCompilerInterface>("KismetCompiler");
-		if (KismetCompilerModulePtr)
-		{
-			KismetCompilerModulePtr->GetCompilers().Remove(this);
-		}
 	}
 
 protected:
@@ -321,28 +317,6 @@ protected:
 		}
 	}
 
-protected:
-
-	static TSharedPtr<FKismetCompilerContext> FactoryCompilerContext(UBlueprint* BP, FCompilerResultsLog& InMessageLog, const FKismetCompilerOptions& InCompileOptions)
-	{
-		return MakeShared<FLevelSequenceDirectorBlueprintCompiler>(CastChecked<ULevelSequenceDirectorBlueprint>(BP), InMessageLog, InCompileOptions, nullptr);
-	}
-
-	virtual bool CanCompile(const UBlueprint* Blueprint) override
-	{
-		return Blueprint && Blueprint->IsA<ULevelSequenceDirectorBlueprint>();
-	}
-
-	virtual void Compile(UBlueprint* Blueprint, const FKismetCompilerOptions& CompileOptions, FCompilerResultsLog& Results, TArray<UObject*>* ObjLoaded) override
-	{
-		if (ULevelSequenceDirectorBlueprint* LevelSequenceDirectorBP = CastChecked<ULevelSequenceDirectorBlueprint>(Blueprint) )
-		{
-			FLevelSequenceDirectorBlueprintCompiler Compiler(LevelSequenceDirectorBP, Results, CompileOptions, ObjLoaded);
-			Compiler.Compile();
-			check(Compiler.NewClass);
-		}
-	}
-
 protected:
 
 	/** Callback for creating a new level sequence asset in the level. */
diff --git a/Engine/Plugins/MovieScene/LevelSequenceEditor/Source/LevelSequenceEditor/Private/Misc/MovieSceneSequenceEditor_LevelSequence.h b/Engine/Plugins/MovieScene/LevelSequenceEditor/Source/LevelSequenceEditor/Private/Misc/MovieSceneSequenceEditor_LevelSequence.h
--- a/Engine/Plugins/MovieScene/LevelSequenceEditor/Source/LevelSequenceEditor/Private/Misc/MovieSceneSequenceEditor_LevelSequence.h
+++ b/Engine/Plugins/MovieScene/LevelSequenceEditor/Source/LevelSequenceEditor/Private/Misc/MovieSceneSequenceEditor_LevelSequence.h
@@ -3,21 +3,14 @@
 #pragma once
 
 #include "MovieSceneSequenceEditor.h"
-#include "LevelSequenceDirectorGeneratedClass.h"
-#include "LevelSequenceDirectorBlueprint.h"
+#include "LevelSequenceDirector.h"
+#include "EdGraphSchema_K2.h"
 #include "Kismet2/KismetEditorUtilities.h"
 #include "Tracks/MovieSceneEventTrack.h"
 #include "K2Node_FunctionEntry.h"
 
 struct FMovieSceneSequenceEditor_LevelSequence : FMovieSceneSequenceEditor
 {
-	FName BlueprintName;
-
-	FMovieSceneSequenceEditor_LevelSequence()
-	{
-		BlueprintName = "SequenceDirectorBlueprint";
-	}
-
 	virtual bool CanCreateEvents(UMovieSceneSequence* InSequence) const
 	{
 		return true;
@@ -39,11 +32,11 @@ struct FMovieSceneSequenceEditor_LevelSequence : FMovieSceneSequenceEditor
 
 		ULevelSequence* LevelSequence = CastChecked<ULevelSequence>(InSequence);
 
-		Blueprint = FKismetEditorUtilities::CreateBlueprint(ULevelSequenceDirector::StaticClass(), InSequence, NAME_None, BPTYPE_Normal, ULevelSequenceDirectorBlueprint::StaticClass(), ULevelSequenceDirectorGeneratedClass::StaticClass());
-		CastChecked<ULevelSequenceDirectorBlueprint>(Blueprint)->OwnerSequence = LevelSequence;
+		FName BlueprintName = "SequenceDirector";
+		Blueprint = FKismetEditorUtilities::CreateBlueprint(ULevelSequenceDirector::StaticClass(), InSequence, BlueprintName, BPTYPE_Normal, UBlueprint::StaticClass(), UBlueprintGeneratedClass::StaticClass());
 
 		LevelSequence->DirectorBlueprint = Blueprint;
-		LevelSequence->DirectorClass = CastChecked<ULevelSequenceDirectorGeneratedClass>(Blueprint->GeneratedClass);
+		LevelSequence->DirectorClass = CastChecked<UBlueprintGeneratedClass>(Blueprint->GeneratedClass);
 		return Blueprint;
 	}
 
diff --git a/Engine/Source/Editor/Kismet/Public/BlueprintAssetHandler.h b/Engine/Source/Editor/Kismet/Public/BlueprintAssetHandler.h
--- a/Engine/Source/Editor/Kismet/Public/BlueprintAssetHandler.h
+++ b/Engine/Source/Editor/Kismet/Public/BlueprintAssetHandler.h
@@ -38,6 +38,17 @@ public:
 	 *@return True if the asset contains a blueprint, false otherwise
 	 */
 	virtual bool AssetContainsBlueprint(const FAssetData& InAssetData) const = 0;
+
+
+	/**
+	 * Check whether the specified asset supports nativization
+	 *
+	 *@param InAsset         The asset that is being queired for nativization support
+	 *@param InBlueprint     The blueprint that is contained within InAsset
+	 *@param OutReason       (Optional) An optional failure text to set
+	 *@return true if the specified asset supports nativization, false otherwise
+	 */
+	virtual bool SupportsNativization(const UObject* InAsset, const UBlueprint* InBlueprint, FText* OutReason) const { return true; }
 };
 
 
diff --git a/Engine/Source/Runtime/Engine/Private/Blueprint.cpp b/Engine/Source/Runtime/Engine/Private/Blueprint.cpp
--- a/Engine/Source/Runtime/Engine/Private/Blueprint.cpp
+++ b/Engine/Source/Runtime/Engine/Private/Blueprint.cpp
@@ -36,6 +36,7 @@
 #include "Curves/CurveBase.h"
 #include "Interfaces/ITargetPlatform.h"
 #include "UObject/MetaData.h"
+#include "BlueprintAssetHandler.h"
 #endif
 #include "Engine/InheritableComponentHandler.h"
 
@@ -830,6 +831,21 @@ bool UBlueprint::SupportsNativization(FText* OutReason) const
 		}
 		return false;
 	}
+	else if (!GetOuter()->IsA<UPackage>())
+	{
+		// If this blueprint is not an asset itself, check whether the asset supports nativization
+		UObject* Asset = GetOuter();
+		while (Asset && !Asset->GetOuter()->IsA<UPackage>())
+		{
+			Asset = Asset->GetOuter();
+		}
+
+		const IBlueprintAssetHandler* Handler = Asset ? FBlueprintAssetHandler::Get().FindHandler(Asset->GetClass()) : nullptr;
+		if (Handler && !Handler->SupportsNativization(Asset, this, OutReason))
+		{
+			return false;
+		}
+	}
 	return true;
 }
 
diff --git a/Engine/Source/Runtime/LevelSequence/Private/LevelSequence.cpp b/Engine/Source/Runtime/LevelSequence/Private/LevelSequence.cpp
--- a/Engine/Source/Runtime/LevelSequence/Private/LevelSequence.cpp
+++ b/Engine/Source/Runtime/LevelSequence/Private/LevelSequence.cpp
@@ -5,7 +5,7 @@
 #include "HAL/IConsoleManager.h"
 #include "Components/ActorComponent.h"
 #include "GameFramework/Actor.h"
-#include "LevelSequenceDirectorGeneratedClass.h"
+#include "LevelSequenceDirector.h"
 #include "Engine/Engine.h"
 #include "MovieScene.h"
 #include "UObject/Package.h"
@@ -486,7 +486,7 @@ UObject* ULevelSequence::CreateDirectorInstance(IMovieScenePlayer& Player)
 	ULevelSequencePlayer* LevelSequencePlayer = Cast<ULevelSequencePlayer>(Player.AsUObject());
 	UObject*              DirectorOuter       = LevelSequencePlayer ? LevelSequencePlayer : Player.GetPlaybackContext();
 
-	if (DirectorClass && DirectorOuter)
+	if (DirectorClass && DirectorOuter && DirectorClass->IsChildOf(ULevelSequenceDirector::StaticClass()))
 	{
 		FName DirectorName = NAME_None;
 
diff --git a/Engine/Source/Runtime/LevelSequence/Private/LevelSequenceDirector.cpp b/Engine/Source/Runtime/LevelSequence/Private/LevelSequenceDirector.cpp
new file mode 100644
--- /dev/null
+++ b/Engine/Source/Runtime/LevelSequence/Private/LevelSequenceDirector.cpp
@@ -0,0 +1,9 @@
+// Copyright 1998-2018 Epic Games, Inc. All Rights Reserved.
+
+#include "LevelSequenceDirector.h"
+#include "Engine/World.h"
+
+UWorld* ULevelSequenceDirector::GetWorld() const
+{
+	return GetTypedOuter<UWorld>();
+}
diff --git a/Engine/Source/Runtime/LevelSequence/Private/LevelSequenceDirectorGeneratedClass.cpp b/Engine/Source/Runtime/LevelSequence/Private/LevelSequenceDirectorGeneratedClass.cpp
deleted file mode 100644
diff --git a/Engine/Source/Runtime/LevelSequence/Public/LevelSequence.h b/Engine/Source/Runtime/LevelSequence/Public/LevelSequence.h
--- a/Engine/Source/Runtime/LevelSequence/Public/LevelSequence.h
+++ b/Engine/Source/Runtime/LevelSequence/Public/LevelSequence.h
@@ -11,7 +11,7 @@
 #include "LevelSequence.generated.h"
 
 class UMovieScene;
-class ULevelSequenceDirectorGeneratedClass;
+class UBlueprintGeneratedClass;
 
 /**
  * Movie scene animation for Actors.
@@ -69,7 +69,7 @@ public:
 	 * Director instances are allocated on-demand one per sequence during evaluation and are used by event tracks for triggering events.
 	 */
 	UPROPERTY()
-	ULevelSequenceDirectorGeneratedClass* DirectorClass;
+	UBlueprintGeneratedClass* DirectorClass;
 
 protected:
 
diff --git a/Engine/Source/Runtime/LevelSequence/Public/LevelSequenceDirector.h b/Engine/Source/Runtime/LevelSequence/Public/LevelSequenceDirector.h
new file mode 100644
--- /dev/null
+++ b/Engine/Source/Runtime/LevelSequence/Public/LevelSequenceDirector.h
@@ -0,0 +1,40 @@
+// Copyright 1998-2018 Epic Games, Inc. All Rights Reserved.
+
+#pragma once
+
+#include "CoreTypes.h"
+#include "Engine/Blueprint.h"
+#include "Engine/BlueprintGeneratedClass.h"
+#include "LevelSequenceDirector.generated.h"
+
+class ULevelSequencePlayer;
+
+UCLASS(Blueprintable)
+class LEVELSEQUENCE_API ULevelSequenceDirector : public UObject
+{
+public:
+	GENERATED_BODY()
+
+	/** Called when this director is created */
+	UFUNCTION(BlueprintImplementableEvent, Category="Sequencer")
+	void OnCreated();
+
+	virtual UWorld* GetWorld() const override;
+
+	/** Pointer to the player that's playing back this director's sequence */
+	UPROPERTY(BlueprintReadOnly, Category="Cinematics")
+	ULevelSequencePlayer* Player;
+};
+
+
+UCLASS()
+class ULegacyLevelSequenceDirectorBlueprint : public UBlueprint
+{
+	GENERATED_BODY()
+
+	ULegacyLevelSequenceDirectorBlueprint(const FObjectInitializer& ObjInit)
+		: Super(ObjInit)
+	{
+		ParentClass = ULevelSequenceDirector::StaticClass();
+	}
+};
\ No newline at end of file
diff --git a/Engine/Source/Runtime/LevelSequence/Public/LevelSequenceDirectorGeneratedClass.h b/Engine/Source/Runtime/LevelSequence/Public/LevelSequenceDirectorGeneratedClass.h
deleted file mode 100644
