Sequencer: Show decimal places up to the zero padding value. Fixed buzzing in time displays by setting the min desired width based on the min/max view range.

diff --git a/Engine/Plugins/MovieScene/LevelSequenceEditor/Source/LevelSequenceEditor/Private/CinematicViewport/SCinematicLevelViewport.cpp b/Engine/Plugins/MovieScene/LevelSequenceEditor/Source/LevelSequenceEditor/Private/CinematicViewport/SCinematicLevelViewport.cpp
--- a/Engine/Plugins/MovieScene/LevelSequenceEditor/Source/LevelSequenceEditor/Private/CinematicViewport/SCinematicLevelViewport.cpp
+++ b/Engine/Plugins/MovieScene/LevelSequenceEditor/Source/LevelSequenceEditor/Private/CinematicViewport/SCinematicLevelViewport.cpp
@@ -28,6 +28,7 @@
 #include "CineCameraComponent.h"
 #include "Math/UnitConversion.h"
 #include "LevelEditorSequencerIntegration.h"
+#include "Fonts/FontMeasure.h"
 
 
 #define LOCTEXT_NAMESPACE "SCinematicLevelViewport"
@@ -212,6 +213,7 @@ void SCinematicLevelViewport::Construct(const FArguments& InArgs)
 					return UIData.OuterResolution.AsDecimal() * UIData.OuterPlayRate.AsInterval(); 
 				})
 				.LinearDeltaSensitivity(25)
+				.MinDesiredWidth(this, &SCinematicLevelViewport::GetPlayTimeMinDesiredWidth)
 			]
 		]
 
@@ -498,6 +500,29 @@ double SCinematicLevelViewport::GetTime() const
 	return 0;
 }
 
+float SCinematicLevelViewport::GetPlayTimeMinDesiredWidth() const
+{
+	ISequencer* Sequencer = GetSequencer();
+	if (Sequencer)
+	{
+		TRange<double> ViewRange = Sequencer->GetViewRange();
+
+		FString LowerBoundStr = Sequencer->GetNumericTypeInterface()->ToString(ViewRange.GetLowerBoundValue());
+		FString UpperBoundStr = Sequencer->GetNumericTypeInterface()->ToString(ViewRange.GetUpperBoundValue());
+
+		const FSlateFontInfo PlayTimeFont = FEditorStyle::GetFontStyle("Sequencer.FixedFont");
+
+		const TSharedRef< FSlateFontMeasure > FontMeasureService = FSlateApplication::Get().GetRenderer()->GetFontMeasureService();
+
+		FVector2D LowerTextSize = FontMeasureService->Measure(LowerBoundStr, PlayTimeFont);
+		FVector2D UpperTextSize = FontMeasureService->Measure(UpperBoundStr, PlayTimeFont);
+
+		return FMath::Max(LowerTextSize.X, UpperTextSize.X);
+	}
+
+	return 0.f;
+}
+
 void SCinematicLevelViewport::CacheDesiredViewportSize(const FGeometry& AllottedGeometry)
 {
 	FVector2D AllowableSpace = AllottedGeometry.GetLocalSize();
diff --git a/Engine/Plugins/MovieScene/LevelSequenceEditor/Source/LevelSequenceEditor/Private/CinematicViewport/SCinematicLevelViewport.h b/Engine/Plugins/MovieScene/LevelSequenceEditor/Source/LevelSequenceEditor/Private/CinematicViewport/SCinematicLevelViewport.h
--- a/Engine/Plugins/MovieScene/LevelSequenceEditor/Source/LevelSequenceEditor/Private/CinematicViewport/SCinematicLevelViewport.h
+++ b/Engine/Plugins/MovieScene/LevelSequenceEditor/Source/LevelSequenceEditor/Private/CinematicViewport/SCinematicLevelViewport.h
@@ -129,6 +129,8 @@ private:
 
 	double GetTime() const;
 
+	float GetPlayTimeMinDesiredWidth() const;
+
 private:
 	
 	/** Widget where the scene viewport is drawn in */
diff --git a/Engine/Source/Editor/Sequencer/Private/SSequencer.cpp b/Engine/Source/Editor/Sequencer/Private/SSequencer.cpp
--- a/Engine/Source/Editor/Sequencer/Private/SSequencer.cpp
+++ b/Engine/Source/Editor/Sequencer/Private/SSequencer.cpp
@@ -76,6 +76,7 @@
 #include "SequencerLog.h"
 #include "MovieSceneCopyableBinding.h"
 #include "MovieSceneCopyableTrack.h"
+#include "Fonts/FontMeasure.h"
 
 #define LOCTEXT_NAMESPACE "Sequencer"
 
@@ -395,6 +396,7 @@ void SSequencer::Construct(const FArguments& InArgs, TSharedRef<FSequencer> InSe
 									.TypeInterface(NumericTypeInterface)
 									.Delta(this, &SSequencer::GetSpinboxDelta)
 									.LinearDeltaSensitivity(25)
+									.MinDesiredWidth(this, &SSequencer::GetPlayTimeMinDesiredWidth)
 								]
 							]
 						]
@@ -2495,6 +2497,24 @@ double SSequencer::GetSpinboxDelta() const
 	return Sequencer->GetDisplayRateDeltaFrameCount();
 }
 
+float SSequencer::GetPlayTimeMinDesiredWidth() const
+{
+	TRange<double> ViewRange = SequencerPtr.Pin()->GetViewRange();
+
+	FString LowerBoundStr = NumericTypeInterface->ToString(ViewRange.GetLowerBoundValue());
+	FString UpperBoundStr = NumericTypeInterface->ToString(ViewRange.GetUpperBoundValue());
+
+	const FSlateFontInfo NormalFont = FCoreStyle::Get().GetFontStyle(TEXT("NormalFont"));
+	
+	const TSharedRef< FSlateFontMeasure > FontMeasureService = FSlateApplication::Get().GetRenderer()->GetFontMeasureService();
+
+	FVector2D LowerTextSize = FontMeasureService->Measure(LowerBoundStr, NormalFont);
+	FVector2D UpperTextSize = FontMeasureService->Measure(UpperBoundStr, NormalFont);
+
+	return FMath::Max(LowerTextSize.X, UpperTextSize.X);
+}
+
+
 bool SSequencer::GetIsSequenceReadOnly() const
 {
 	TSharedPtr<FSequencer> Sequencer = SequencerPtr.Pin();
diff --git a/Engine/Source/Editor/Sequencer/Private/SSequencer.h b/Engine/Source/Editor/Sequencer/Private/SSequencer.h
--- a/Engine/Source/Editor/Sequencer/Private/SSequencer.h
+++ b/Engine/Source/Editor/Sequencer/Private/SSequencer.h
@@ -442,6 +442,9 @@ private:
 	/** Controls how fast Spinboxes change values. */
 	double GetSpinboxDelta() const;
 
+	/** Get minimum desired width of the current time spin box */
+	float GetPlayTimeMinDesiredWidth() const;
+
 	bool GetIsSequenceReadOnly() const;
 	void OnSetSequenceReadOnly(ECheckBoxState CheckBoxState);
 
diff --git a/Engine/Source/Runtime/TimeManagement/Public/FrameNumberNumericInterface.h b/Engine/Source/Runtime/TimeManagement/Public/FrameNumberNumericInterface.h
--- a/Engine/Source/Runtime/TimeManagement/Public/FrameNumberNumericInterface.h
+++ b/Engine/Source/Runtime/TimeManagement/Public/FrameNumberNumericInterface.h
@@ -8,6 +8,9 @@
 #include "Misc/FrameRate.h"
 #include "FrameNumberTimeEvaluator.h"
 #include "Misc/Timecode.h"
+#include "Misc/ExpressionParserTypes.h"
+#include "Math/BasicMathExpressionEvaluator.h"
+#include "Internationalization/FastDecimalFormat.h"
 
 /**
 * A large portion of the Sequencer UI is built around using SpinBox and NumericBox because the framerate
@@ -78,7 +81,12 @@ struct FFrameNumberInterface : public INumericTypeInterface<double>
 		case EFrameNumberDisplayFormats::Seconds:
 		{
 			double TimeInSeconds = SourceFrameRate.AsSeconds(FFrameTime::FromDecimal(Value));
-			return FString::Printf(TEXT("%.2f s"), TimeInSeconds);
+
+			static const FNumberFormattingOptions NumberFormattingOptions = FNumberFormattingOptions()
+				.SetUseGrouping(false)
+				.SetMinimumFractionalDigits(2)
+				.SetMaximumFractionalDigits(ZeroPadFramesAttr.Get());
+			return FastDecimalFormat::NumberToString(TimeInSeconds, ExpressionParser::GetLocalizedNumberFormattingRules(), NumberFormattingOptions);
 		}
 		case EFrameNumberDisplayFormats::NonDropFrameTimecode:
 		case EFrameNumberDisplayFormats::DropFrameTimecode:
