Fix crash when trying to create a level sequence with a locked map.

diff --git a/Engine/Plugins/MovieScene/LevelSequenceEditor/Source/LevelSequenceEditor/Private/LevelSequenceEditorModule.cpp b/Engine/Plugins/MovieScene/LevelSequenceEditor/Source/LevelSequenceEditor/Private/LevelSequenceEditorModule.cpp
--- a/Engine/Plugins/MovieScene/LevelSequenceEditor/Source/LevelSequenceEditor/Private/LevelSequenceEditorModule.cpp
+++ b/Engine/Plugins/MovieScene/LevelSequenceEditor/Source/LevelSequenceEditor/Private/LevelSequenceEditorModule.cpp
@@ -339,7 +339,12 @@ protected:
 			return;
 		}
 
-		ALevelSequenceActor* NewActor = CastChecked<ALevelSequenceActor>(GEditor->UseActorFactory(ActorFactory, FAssetData(NewAsset), &FTransform::Identity));
+		AActor* Actor = GEditor->UseActorFactory(ActorFactory, FAssetData(NewAsset), &FTransform::Identity);
+		if (Actor == nullptr)
+		{
+			return;
+		}
+		ALevelSequenceActor* NewActor = CastChecked<ALevelSequenceActor>(Actor);
 		if (GCurrentLevelEditingViewportClient != nullptr && GCurrentLevelEditingViewportClient->IsPerspective())
 		{
 			GEditor->MoveActorInFrontOfCamera(*NewActor, GCurrentLevelEditingViewportClient->GetViewLocation(), GCurrentLevelEditingViewportClient->GetViewRotation().Vector());
