Sequencer: Add setting to determine whether to show the outliner info column

diff --git a/Engine/Plugins/MovieScene/LevelSequenceEditor/Source/LevelSequenceEditor/Private/LevelSequenceEditorToolkit.cpp b/Engine/Plugins/MovieScene/LevelSequenceEditor/Source/LevelSequenceEditor/Private/LevelSequenceEditorToolkit.cpp
--- a/Engine/Plugins/MovieScene/LevelSequenceEditor/Source/LevelSequenceEditor/Private/LevelSequenceEditorToolkit.cpp
+++ b/Engine/Plugins/MovieScene/LevelSequenceEditor/Source/LevelSequenceEditor/Private/LevelSequenceEditorToolkit.cpp
@@ -203,11 +203,14 @@ void FLevelSequenceEditorToolkit::Initialize(const EToolkitMode::Type Mode, cons
 	FLevelEditorModule& LevelEditorModule = FModuleManager::LoadModuleChecked<FLevelEditorModule>("LevelEditor");
 
 	// Reopen the scene outliner so that is refreshed with the sequencer info column
-	TSharedPtr<FTabManager> LevelEditorTabManager = LevelEditorModule.GetLevelEditorTabManager();
-	if (LevelEditorTabManager->FindExistingLiveTab(FName("LevelEditorSceneOutliner")).IsValid())
+	if (Sequencer->GetSequencerSettings()->GetShowOutlinerInfoColumn())
 	{
-		LevelEditorTabManager->InvokeTab(FName("LevelEditorSceneOutliner"))->RequestCloseTab();
-		LevelEditorTabManager->InvokeTab(FName("LevelEditorSceneOutliner"));
+		TSharedPtr<FTabManager> LevelEditorTabManager = LevelEditorModule.GetLevelEditorTabManager();
+		if (LevelEditorTabManager->FindExistingLiveTab(FName("LevelEditorSceneOutliner")).IsValid())
+		{
+			LevelEditorTabManager->InvokeTab(FName("LevelEditorSceneOutliner"))->RequestCloseTab();
+			LevelEditorTabManager->InvokeTab(FName("LevelEditorSceneOutliner"));
+		}
 	}
 	
 	// Now Attach so this window will apear in the correct front first order
diff --git a/Engine/Source/Editor/Sequencer/Private/LevelEditorSequencerIntegration.cpp b/Engine/Source/Editor/Sequencer/Private/LevelEditorSequencerIntegration.cpp
--- a/Engine/Source/Editor/Sequencer/Private/LevelEditorSequencerIntegration.cpp
+++ b/Engine/Source/Editor/Sequencer/Private/LevelEditorSequencerIntegration.cpp
@@ -219,7 +219,6 @@ void FLevelEditorSequencerIntegration::Initialize()
 	ActivateDetailHandler();
 	ActivateSequencerEditorMode();
 	BindLevelEditorCommands();
-	AttachOutlinerColumn();
 
 	{
 		FPropertyEditorModule& EditModule = FModuleManager::Get().GetModuleChecked<FPropertyEditorModule>("PropertyEditor");
@@ -860,6 +859,18 @@ TSharedRef< ISceneOutlinerColumn > FLevelEditorSequencerIntegration::CreateSeque
 
 void FLevelEditorSequencerIntegration::AttachOutlinerColumn()
 {
+	for (const FSequencerAndOptions& SequencerAndOptions : BoundSequencers)
+	{
+		TSharedPtr<FSequencer> Pinned = SequencerAndOptions.Sequencer.Pin();
+		if (Pinned.IsValid())
+		{
+			if (!Pinned.Get()->GetSequencerSettings()->GetShowOutlinerInfoColumn())
+			{
+				return;
+			}
+		}
+	}
+
 	FSceneOutlinerModule& SceneOutlinerModule = FModuleManager::LoadModuleChecked< FSceneOutlinerModule >("SceneOutliner");
 
 	SceneOutliner::FColumnInfo ColumnInfo(SceneOutliner::EColumnVisibility::Visible, 15, 
@@ -1038,6 +1049,7 @@ void FLevelEditorSequencerIntegration::AddSequencer(TSharedRef<ISequencer> InSeq
 	}
 
 	ActivateRealtimeViewports();
+	AttachOutlinerColumn();
 }
 
 void FLevelEditorSequencerIntegration::OnSequencerReceivedFocus(TSharedRef<ISequencer> InSequencer)
diff --git a/Engine/Source/Editor/Sequencer/Private/SequencerSettings.cpp b/Engine/Source/Editor/Sequencer/Private/SequencerSettings.cpp
--- a/Engine/Source/Editor/Sequencer/Private/SequencerSettings.cpp
+++ b/Engine/Source/Editor/Sequencer/Private/SequencerSettings.cpp
@@ -47,6 +47,7 @@ USequencerSettings::USequencerSettings( const FObjectInitializer& ObjectInitiali
 	bRerunConstructionScripts = true;
 	bVisualizePreAndPostRoll = true;
 	TrajectoryPathCap = 250;
+	bShowOutlinerInfoColumn = true;
 	CurveVisibility = ECurveEditorCurveVisibility::SelectedCurves;
 	FrameNumberDisplayFormat = EFrameNumberDisplayFormats::Seconds;
 }
@@ -628,6 +629,20 @@ void USequencerSettings::SetCompileDirectorOnEvaluate(bool bInCompileDirectorOnE
 	}
 }
 
+bool USequencerSettings::GetShowOutlinerInfoColumn() const
+{
+	return bShowOutlinerInfoColumn;
+}
+
+void USequencerSettings::SetShowOutlinerInfoColumn(bool bInShowOutlinerInfoColumn)
+{
+	if (bInShowOutlinerInfoColumn != bShowOutlinerInfoColumn)
+	{
+		bShowOutlinerInfoColumn = bInShowOutlinerInfoColumn;
+		SaveConfig();
+	}
+}
+
 ECurveEditorCurveVisibility USequencerSettings::GetCurveVisibility() const
 {
 	return CurveVisibility;
diff --git a/Engine/Source/Editor/Sequencer/Public/SequencerSettings.h b/Engine/Source/Editor/Sequencer/Public/SequencerSettings.h
--- a/Engine/Source/Editor/Sequencer/Public/SequencerSettings.h
+++ b/Engine/Source/Editor/Sequencer/Public/SequencerSettings.h
@@ -321,6 +321,11 @@ public:
 
 	uint32 GetTrajectoryPathCap() const { return TrajectoryPathCap; }
 
+	/** Gets whether to show the sequencer outliner info column */
+	bool GetShowOutlinerInfoColumn() const;
+	/** Sets whether to show the sequencer outliner info column */
+	void SetShowOutlinerInfoColumn(bool bInShowOutlinerInfoColumn);
+
 	/** Gets the current curve visibility. */
 	ECurveEditorCurveVisibility GetCurveVisibility() const;
 	/** Sets the current curve visibility. */
@@ -496,10 +501,14 @@ protected:
 	bool bCompileDirectorOnEvaluate;
 
 	/** Specifies the maximum number of keys to draw when rendering trajectories in viewports */
-	UPROPERTY(config, EditAnywhere, Category=General)
+	UPROPERTY(config, EditAnywhere, Category = General)
 	uint32 TrajectoryPathCap;
 
-	/** What format do we display time in to the user? */	
+	/** Whether to show the sequencer outliner info column */
+	UPROPERTY(config, EditAnywhere, Category = General)
+	bool bShowOutlinerInfoColumn;
+
+	/** What format do we display time in to the user? */
 	UPROPERTY(config, EditAnywhere, Category=General)
 	EFrameNumberDisplayFormats FrameNumberDisplayFormat;
 
