Sequencer: Remove showing sub sections in the cinematic viewport because it's ambiguous as to which subsection to choose if there are multiple.

diff --git a/Engine/Plugins/MovieScene/LevelSequenceEditor/Source/LevelSequenceEditor/Private/CinematicViewport/SCinematicLevelViewport.cpp b/Engine/Plugins/MovieScene/LevelSequenceEditor/Source/LevelSequenceEditor/Private/CinematicViewport/SCinematicLevelViewport.cpp
--- a/Engine/Plugins/MovieScene/LevelSequenceEditor/Source/LevelSequenceEditor/Private/CinematicViewport/SCinematicLevelViewport.cpp
+++ b/Engine/Plugins/MovieScene/LevelSequenceEditor/Source/LevelSequenceEditor/Private/CinematicViewport/SCinematicLevelViewport.cpp
@@ -17,8 +17,6 @@
 #include "MovieScene.h"
 #include "ISequencer.h"
 #include "MovieSceneSequence.h"
-#include "Sections/MovieSceneSubSection.h"
-#include "Tracks/MovieSceneSubTrack.h"
 #include "Tracks/MovieSceneCinematicShotTrack.h"
 #include "Sections/MovieSceneCinematicShotSection.h"
 #include "SequencerKeyCollection.h"
@@ -652,12 +650,8 @@ void SCinematicLevelViewport::Tick(const FGeometry& AllottedGeometry, const doub
 		return;
 	}
 
-	// Find either a cinematic shot track or a sub track
-	UMovieSceneSubTrack* SubTrack = Cast<UMovieSceneSubTrack>(Sequence->GetMovieScene()->FindMasterTrack(UMovieSceneCinematicShotTrack::StaticClass()));
-	if (!SubTrack)
-	{
-		SubTrack = Cast<UMovieSceneSubTrack>(Sequence->GetMovieScene()->FindMasterTrack(UMovieSceneSubTrack::StaticClass()));
-	}
+	// Find the cinematic shot track
+	UMovieSceneCinematicShotTrack* CinematicShotTrack = Cast<UMovieSceneCinematicShotTrack>(Sequence->GetMovieScene()->FindMasterTrack(UMovieSceneCinematicShotTrack::StaticClass()));
 
 	const FFrameRate OuterResolution = Sequencer->GetFocusedTickResolution();
 	const FFrameRate OuterPlayRate   = Sequencer->GetFocusedDisplayRate();
@@ -666,15 +660,14 @@ void SCinematicLevelViewport::Tick(const FGeometry& AllottedGeometry, const doub
 	UIData.OuterResolution = OuterResolution;
 	UIData.OuterPlayRate = OuterPlayRate;
 
-
-	UMovieSceneSubSection* SubSection = nullptr;
-	if (SubTrack)
+	UMovieSceneCinematicShotSection* CinematicShotSection = nullptr;
+	if (CinematicShotTrack)
 	{
-		for (UMovieSceneSection* Section : SubTrack->GetAllSections())
+		for (UMovieSceneSection* Section : CinematicShotTrack->GetAllSections())
 		{
 			if (Section->GetRange().Contains(OuterTime.FrameNumber))
 			{
-				SubSection = CastChecked<UMovieSceneSubSection>(Section);
+				CinematicShotSection = CastChecked<UMovieSceneCinematicShotSection>(Section);
 			}
 		}
 	}
@@ -683,11 +676,11 @@ void SCinematicLevelViewport::Tick(const FGeometry& AllottedGeometry, const doub
 
 	TSharedPtr<INumericTypeInterface<double>> TimeDisplayFormatInterface = Sequencer->GetNumericTypeInterface();
 
-	UMovieSceneSequence* SubSequence = SubSection ? SubSection->GetSequence() : nullptr;
+	UMovieSceneSequence* SubSequence = CinematicShotSection ? CinematicShotSection->GetSequence() : nullptr;
 	if (SubSequence)
 	{
 		FFrameRate                   InnerResolution       = SubSequence->GetMovieScene()->GetTickResolution();
-		FMovieSceneSequenceTransform OuterToInnerTransform = SubSection->OuterToInnerTransform();
+		FMovieSceneSequenceTransform OuterToInnerTransform = CinematicShotSection->OuterToInnerTransform();
 		const FFrameTime             InnerShotPosition	   = OuterTime * OuterToInnerTransform;
 
 		UIData.LocalPlaybackTime = FText::Format(
@@ -695,15 +688,10 @@ void SCinematicLevelViewport::Tick(const FGeometry& AllottedGeometry, const doub
 			FText::FromString(TimeDisplayFormatInterface->ToString(InnerShotPosition.GetFrame().Value))
 		);
 
-		UMovieSceneCinematicShotSection* CinematicShotSection = Cast<UMovieSceneCinematicShotSection>(SubSection);
 		if (CinematicShotSection)
 		{
 			UIData.ShotName = FText::FromString(CinematicShotSection->GetShotDisplayName());
 		}
-		else if (SubSection->GetSequence() != nullptr)
-		{
-			UIData.ShotName = SubSection->GetSequence()->GetDisplayName();
-		}
 	}
 	else
 	{
