Sequencer: Fix static analysis warning

diff --git a/Engine/Plugins/MovieScene/LevelSequenceEditor/Source/LevelSequenceEditor/Private/CinematicViewport/SCinematicLevelViewport.cpp b/Engine/Plugins/MovieScene/LevelSequenceEditor/Source/LevelSequenceEditor/Private/CinematicViewport/SCinematicLevelViewport.cpp
--- a/Engine/Plugins/MovieScene/LevelSequenceEditor/Source/LevelSequenceEditor/Private/CinematicViewport/SCinematicLevelViewport.cpp
+++ b/Engine/Plugins/MovieScene/LevelSequenceEditor/Source/LevelSequenceEditor/Private/CinematicViewport/SCinematicLevelViewport.cpp
@@ -680,7 +680,7 @@ void SCinematicLevelViewport::Tick(const FGeometry& AllottedGeometry, const doub
 	if (SubSequence)
 	{
 		FFrameRate                   InnerResolution       = SubSequence->GetMovieScene()->GetTickResolution();
-		FMovieSceneSequenceTransform OuterToInnerTransform = CinematicShotSection->OuterToInnerTransform();
+		FMovieSceneSequenceTransform OuterToInnerTransform = CinematicShotSection ? CinematicShotSection->OuterToInnerTransform() : FMovieSceneSequenceTransform();
 		const FFrameTime             InnerShotPosition	   = OuterTime * OuterToInnerTransform;
 
 		UIData.LocalPlaybackTime = FText::Format(
