Sequencer: Fix not being to add components to unselected actors

diff --git a/Engine/Plugins/MovieScene/LevelSequenceEditor/Source/LevelSequenceEditor/Private/LevelSequenceEditorToolkit.cpp b/Engine/Plugins/MovieScene/LevelSequenceEditor/Source/LevelSequenceEditor/Private/LevelSequenceEditorToolkit.cpp
--- a/Engine/Plugins/MovieScene/LevelSequenceEditor/Source/LevelSequenceEditor/Private/LevelSequenceEditorToolkit.cpp
+++ b/Engine/Plugins/MovieScene/LevelSequenceEditor/Source/LevelSequenceEditor/Private/LevelSequenceEditorToolkit.cpp
@@ -559,21 +559,34 @@ void FLevelSequenceEditorToolkit::HandleAddComponentActionExecute(UActorComponen
 	const FScopedTransaction Transaction(LOCTEXT("AddComponent", "Add Component"));
 
 	FString ComponentName = Component->GetName();
+
+	TArray<UActorComponent*> ActorComponents;
+	ActorComponents.Add(Component);
+
 	USelection* SelectedActors = GEditor->GetSelectedActors();
-	for (FSelectionIterator Iter(*SelectedActors); Iter; ++Iter)
+	if (SelectedActors && SelectedActors->Num() > 0)
 	{
-		AActor* Actor = CastChecked<AActor>(*Iter);
-
-		TArray<UActorComponent*> OutComponents;
-		Actor->GetComponents(OutComponents);
-		for (UActorComponent* ActorComponent : OutComponents)
+		for (FSelectionIterator Iter(*SelectedActors); Iter; ++Iter)
 		{
-			if (ActorComponent->GetName() == ComponentName)
+			AActor* Actor = CastChecked<AActor>(*Iter);
+
+			TArray<UActorComponent*> OutActorComponents;
+			Actor->GetComponents(OutActorComponents);
+	
+			for (UActorComponent* ActorComponent : OutActorComponents)
 			{
-				Sequencer->GetHandleToObject(ActorComponent);
+				if (ActorComponent->GetName() == ComponentName)
+				{
+					ActorComponents.AddUnique(ActorComponent);
+				}
 			}
 		}
 	}
+
+	for (UActorComponent* ActorComponent : ActorComponents)
+	{
+		Sequencer->GetHandleToObject(ActorComponent);
+	}
 }
 
 
@@ -590,29 +603,42 @@ void FLevelSequenceEditorToolkit::HandleAddComponentMaterialActionExecute(UPrimi
 	FocusedMovieScene->Modify();
 
 	FString ComponentName = Component->GetName();
+
+	TArray<UActorComponent*> ActorComponents;
+	ActorComponents.Add(Component);
+
 	USelection* SelectedActors = GEditor->GetSelectedActors();
-	for (FSelectionIterator Iter(*SelectedActors); Iter; ++Iter)
+	if (SelectedActors && SelectedActors->Num() > 0)
 	{
-		AActor* Actor = CastChecked<AActor>(*Iter);
-
-		TArray<UActorComponent*> OutComponents;
-		Actor->GetComponents(OutComponents);
-		for (UActorComponent* ActorComponent : OutComponents)
+		for (FSelectionIterator Iter(*SelectedActors); Iter; ++Iter)
 		{
-			if (ActorComponent->GetName() == ComponentName)
+			AActor* Actor = CastChecked<AActor>(*Iter);
+
+			TArray<UActorComponent*> OutActorComponents;
+			Actor->GetComponents(OutActorComponents);
+			for (UActorComponent* ActorComponent : OutActorComponents)
 			{
-				FGuid ObjectHandle = Sequencer->GetHandleToObject(ActorComponent);
-				FName IndexName(*FString::FromInt(MaterialIndex));
-				if (FocusedMovieScene->FindTrack(UMovieSceneComponentMaterialTrack::StaticClass(), ObjectHandle, IndexName) == nullptr)
+				if (ActorComponent->GetName() == ComponentName)
 				{
-					UMovieSceneComponentMaterialTrack* MaterialTrack = FocusedMovieScene->AddTrack<UMovieSceneComponentMaterialTrack>(ObjectHandle);
-					MaterialTrack->Modify();
-					MaterialTrack->SetMaterialIndex(MaterialIndex);
+					ActorComponents.AddUnique(ActorComponent);
 				}
 			}
 		}
-		Sequencer->NotifyMovieSceneDataChanged( EMovieSceneDataChangeType::MovieSceneStructureItemAdded );
 	}
+
+	for (UActorComponent* ActorComponent : ActorComponents)
+	{
+		FGuid ObjectHandle = Sequencer->GetHandleToObject(ActorComponent);
+		FName IndexName(*FString::FromInt(MaterialIndex));
+		if (FocusedMovieScene->FindTrack(UMovieSceneComponentMaterialTrack::StaticClass(), ObjectHandle, IndexName) == nullptr)
+		{
+			UMovieSceneComponentMaterialTrack* MaterialTrack = FocusedMovieScene->AddTrack<UMovieSceneComponentMaterialTrack>(ObjectHandle);
+			MaterialTrack->Modify();
+			MaterialTrack->SetMaterialIndex(MaterialIndex);
+		}
+	}
+
+	Sequencer->NotifyMovieSceneDataChanged(EMovieSceneDataChangeType::MovieSceneStructureItemAdded);
 }
 
 
