Sequencer: Fix crash on undo of duplicating a shot with a spawnable

diff --git a/Engine/Plugins/MovieScene/LevelSequenceEditor/Source/LevelSequenceEditor/Private/Misc/LevelSequenceEditorSpawnRegister.cpp b/Engine/Plugins/MovieScene/LevelSequenceEditor/Source/LevelSequenceEditor/Private/Misc/LevelSequenceEditorSpawnRegister.cpp
--- a/Engine/Plugins/MovieScene/LevelSequenceEditor/Source/LevelSequenceEditor/Private/Misc/LevelSequenceEditorSpawnRegister.cpp
+++ b/Engine/Plugins/MovieScene/LevelSequenceEditor/Source/LevelSequenceEditor/Private/Misc/LevelSequenceEditorSpawnRegister.cpp
@@ -93,7 +93,7 @@ void FLevelSequenceEditorSpawnRegister::PreDestroyObject(UObject& Object, const
 	TSharedPtr<ISequencer> Sequencer = WeakSequencer.Pin();
 
 	UMovieSceneSequence*  Sequence      = Sequencer.IsValid() ? Sequencer->GetEvaluationTemplate().GetSequence(TemplateID) : nullptr;
-	FMovieSceneSpawnable* Spawnable     = Sequence ? Sequence->GetMovieScene()->FindSpawnable(BindingId) : nullptr;
+	FMovieSceneSpawnable* Spawnable     = Sequence && Sequence->GetMovieScene() ? Sequence->GetMovieScene()->FindSpawnable(BindingId) : nullptr;
 	UObject*              SpawnedObject = FindSpawnedObject(BindingId, TemplateID);
 
 	if (SpawnedObject && Spawnable)
