Sequencer: Added GetCurrentLevelSequence

diff --git a/Engine/Plugins/MovieScene/LevelSequenceEditor/Source/LevelSequenceEditor/Private/LevelSequenceEditorBlueprintLibrary.cpp b/Engine/Plugins/MovieScene/LevelSequenceEditor/Source/LevelSequenceEditor/Private/LevelSequenceEditorBlueprintLibrary.cpp
--- a/Engine/Plugins/MovieScene/LevelSequenceEditor/Source/LevelSequenceEditor/Private/LevelSequenceEditorBlueprintLibrary.cpp
+++ b/Engine/Plugins/MovieScene/LevelSequenceEditor/Source/LevelSequenceEditor/Private/LevelSequenceEditorBlueprintLibrary.cpp
@@ -18,6 +18,15 @@ bool ULevelSequenceEditorBlueprintLibrary::OpenLevelSequence(ULevelSequence* Lev
 	return FAssetEditorManager::Get().OpenEditorForAsset(LevelSequence);
 }
 
+ULevelSequence* ULevelSequenceEditorBlueprintLibrary::GetCurrentLevelSequence()
+{
+	if (CurrentSequencer.IsValid())
+	{
+		return Cast<ULevelSequence>(CurrentSequencer.Pin()->GetRootMovieSceneSequence());
+	}
+	return nullptr;
+}
+
 void ULevelSequenceEditorBlueprintLibrary::CloseLevelSequence()
 {
 	if (CurrentSequencer.IsValid())
diff --git a/Engine/Plugins/MovieScene/LevelSequenceEditor/Source/LevelSequenceEditor/Private/LevelSequenceEditorBlueprintLibrary.h b/Engine/Plugins/MovieScene/LevelSequenceEditor/Source/LevelSequenceEditor/Private/LevelSequenceEditorBlueprintLibrary.h
--- a/Engine/Plugins/MovieScene/LevelSequenceEditor/Source/LevelSequenceEditor/Private/LevelSequenceEditorBlueprintLibrary.h
+++ b/Engine/Plugins/MovieScene/LevelSequenceEditor/Source/LevelSequenceEditor/Private/LevelSequenceEditorBlueprintLibrary.h
@@ -22,6 +22,12 @@ public:
 	UFUNCTION(BlueprintCallable, Category = "Level Sequence Editor")
 	static bool OpenLevelSequence(ULevelSequence* LevelSequence);
 
+	/*
+	 * Get the currently opened level sequence asset
+	 */
+	UFUNCTION(BlueprintCallable, Category = "Level Sequence Editor")
+	static ULevelSequence* GetCurrentLevelSequence();
+
 	/*
 	 * Close
 	 */
