Sequencer: Fix spawnable default state not always getting saved. OnObjectModified saves the default state of the object if it is in the list of ModifiedObjects. However, the spawnable is removed from ModifiedObjects in OnPreObjectSaved (this doesn't seem right and is being removed in this CL)

diff --git a/Engine/Plugins/MovieScene/LevelSequenceEditor/Source/LevelSequenceEditor/Private/Misc/LevelSequenceEditorSpawnRegister.cpp b/Engine/Plugins/MovieScene/LevelSequenceEditor/Source/LevelSequenceEditor/Private/Misc/LevelSequenceEditorSpawnRegister.cpp
--- a/Engine/Plugins/MovieScene/LevelSequenceEditor/Source/LevelSequenceEditor/Private/Misc/LevelSequenceEditorSpawnRegister.cpp
+++ b/Engine/Plugins/MovieScene/LevelSequenceEditor/Source/LevelSequenceEditor/Private/Misc/LevelSequenceEditorSpawnRegister.cpp
@@ -72,8 +72,8 @@ UObject* FLevelSequenceEditorSpawnRegister::SpawnObject(FMovieSceneSpawnable& Sp
 	
 	if (AActor* NewActor = Cast<AActor>(NewObject))
 	{
-		// Add an entry to the modified objects map to keep track of when this object has been modified
-		ModifiedObjects.Add(NewActor, FTrackedObjectState(TemplateID, Spawnable.GetGuid()));
+		// Add an entry to the tracked objects map to keep track of this object (so that it can be saved when modified)
+		TrackedObjects.Add(NewActor, FTrackedObjectState(TemplateID, Spawnable.GetGuid()));
 
 		// Select the actor if we think it should be selected
 		if (SelectedSpawnedObjects.Contains(FMovieSceneSpawnRegisterKey(TemplateID, Spawnable.GetGuid())))
@@ -98,7 +98,7 @@ void FLevelSequenceEditorSpawnRegister::PreDestroyObject(UObject& Object, const
 
 	if (SpawnedObject && Spawnable)
 	{
-		const FTrackedObjectState* TrackedState = ModifiedObjects.Find(&Object);
+		const FTrackedObjectState* TrackedState = TrackedObjects.Find(&Object);
 		if (TrackedState && TrackedState->bHasBeenModified)
 		{
 			// SaveDefaultSpawnableState will reset bHasBeenModified to false
@@ -117,7 +117,7 @@ void FLevelSequenceEditorSpawnRegister::PreDestroyObject(UObject& Object, const
 	}
 
 	FObjectKey ThisObject(&Object);
-	ModifiedObjects.Remove(ThisObject);
+	TrackedObjects.Remove(ThisObject);
 
 	FLevelSequenceSpawnRegister::PreDestroyObject(Object, BindingId, TemplateID);
 }
@@ -189,7 +189,7 @@ void FLevelSequenceEditorSpawnRegister::SaveDefaultSpawnableStateImpl(FMovieScen
 	// Copy the template
 	Spawnable.CopyObjectTemplate(*SpawnedObject, *Sequence);
 
-	if (FTrackedObjectState* TrackedState = ModifiedObjects.Find(SpawnedObject))
+	if (FTrackedObjectState* TrackedState = TrackedObjects.Find(SpawnedObject))
 	{
 		TrackedState->bHasBeenModified = false;
 	}
@@ -243,10 +243,10 @@ void FLevelSequenceEditorSpawnRegister::OnObjectsReplaced(const TMap<UObject*, U
 
 void FLevelSequenceEditorSpawnRegister::OnObjectModified(UObject* ModifiedObject)
 {
-	FTrackedObjectState* TrackedState = ModifiedObjects.Find(ModifiedObject);
+	FTrackedObjectState* TrackedState = TrackedObjects.Find(ModifiedObject);
 	while (!TrackedState && ModifiedObject)
 	{
-		TrackedState = ModifiedObjects.Find(ModifiedObject);
+		TrackedState = TrackedObjects.Find(ModifiedObject);
 		ModifiedObject = ModifiedObject->GetOuter();
 	}
 
@@ -278,9 +278,7 @@ void FLevelSequenceEditorSpawnRegister::OnPreObjectSaved(UObject* Object)
 
 		if (Sequencer.IsValid())
 		{
-			TArray<FObjectKey> ObjectsThatHaveBeenSaved;
-
-			for (const TTuple<FObjectKey, FTrackedObjectState>& Pair : ModifiedObjects)
+			for (const TTuple<FObjectKey, FTrackedObjectState>& Pair : TrackedObjects)
 			{
 				UObject* SpawnedObject = Pair.Key.ResolveObjectPtr();
 				UMovieSceneSequence*  ThisSequence = Sequencer->GetEvaluationTemplate().GetSequence(Pair.Value.TemplateID);
@@ -289,15 +287,8 @@ void FLevelSequenceEditorSpawnRegister::OnPreObjectSaved(UObject* Object)
 				if (SpawnedObject && Spawnable && ThisSequence == SequenceBeingSaved)
 				{
 					SaveDefaultSpawnableStateImpl(*Spawnable, ThisSequence, SpawnedObject, *Sequencer);
-
-					ObjectsThatHaveBeenSaved.Add(Pair.Key);
 				}
 			}
-
-			for (FObjectKey ObjectKey : ObjectsThatHaveBeenSaved)
-			{
-				ModifiedObjects.Remove(ObjectKey);
-			}
 		}
 	}
 }
diff --git a/Engine/Plugins/MovieScene/LevelSequenceEditor/Source/LevelSequenceEditor/Private/Misc/LevelSequenceEditorSpawnRegister.h b/Engine/Plugins/MovieScene/LevelSequenceEditor/Source/LevelSequenceEditor/Private/Misc/LevelSequenceEditorSpawnRegister.h
--- a/Engine/Plugins/MovieScene/LevelSequenceEditor/Source/LevelSequenceEditor/Private/Misc/LevelSequenceEditorSpawnRegister.h
+++ b/Engine/Plugins/MovieScene/LevelSequenceEditor/Source/LevelSequenceEditor/Private/Misc/LevelSequenceEditorSpawnRegister.h
@@ -87,8 +87,8 @@ private:
 	/** Set of spawn register keys for objects that should be selected if they are spawned. */
 	TSet<FMovieSceneSpawnRegisterKey> SelectedSpawnedObjects;
 
-	/** Map from a sequenceID to an array of objects that have been modified */
-	TMap<FObjectKey, FTrackedObjectState> ModifiedObjects;
+	/** Map from a sequenceID to an array of objects that have been tracked */
+	TMap<FObjectKey, FTrackedObjectState> TrackedObjects;
 
 	/** Set of UMovieSceneSequences that this register has spawned objects for that are modified */
 	TSet<FObjectKey> SequencesWithModifiedObjects;
