Sequencer: Remove unused IsPaused

diff --git a/Engine/Plugins/MovieScene/LevelSequenceEditor/Source/LevelSequenceEditor/Private/LevelSequenceEditorBlueprintLibrary.cpp b/Engine/Plugins/MovieScene/LevelSequenceEditor/Source/LevelSequenceEditor/Private/LevelSequenceEditorBlueprintLibrary.cpp
--- a/Engine/Plugins/MovieScene/LevelSequenceEditor/Source/LevelSequenceEditor/Private/LevelSequenceEditorBlueprintLibrary.cpp
+++ b/Engine/Plugins/MovieScene/LevelSequenceEditor/Source/LevelSequenceEditor/Private/LevelSequenceEditorBlueprintLibrary.cpp
@@ -84,15 +84,6 @@ bool ULevelSequenceEditorBlueprintLibrary::IsPlaying()
 	return false;
 }
 
-bool ULevelSequenceEditorBlueprintLibrary::IsPaused()
-{
-	if (CurrentSequencer.IsValid())
-	{
-		return CurrentSequencer.Pin()->GetPlaybackStatus() == EMovieScenePlayerStatus::Paused;
-	}
-	return false;
-}
-
 void ULevelSequenceEditorBlueprintLibrary::SetSequencer(TSharedRef<ISequencer> InSequencer)
 {
 	CurrentSequencer = TWeakPtr<ISequencer>(InSequencer);
diff --git a/Engine/Plugins/MovieScene/LevelSequenceEditor/Source/LevelSequenceEditor/Private/LevelSequenceEditorBlueprintLibrary.h b/Engine/Plugins/MovieScene/LevelSequenceEditor/Source/LevelSequenceEditor/Private/LevelSequenceEditorBlueprintLibrary.h
--- a/Engine/Plugins/MovieScene/LevelSequenceEditor/Source/LevelSequenceEditor/Private/LevelSequenceEditorBlueprintLibrary.h
+++ b/Engine/Plugins/MovieScene/LevelSequenceEditor/Source/LevelSequenceEditor/Private/LevelSequenceEditorBlueprintLibrary.h
@@ -66,10 +66,6 @@ public:
 	UFUNCTION(BlueprintPure, Category = "Level Sequence Editor")
 	static bool IsPlaying();
 
-	/** Check whether the sequence is paused. */
-	UFUNCTION(BlueprintPure, Category = "Level Sequence Editor")
-	static bool IsPaused();
-
 public:
 
 	/*
