Sequencer: Ensure that sequence directors are never marked RF_Standalone for level sequences

diff --git a/Engine/Plugins/MovieScene/LevelSequenceEditor/Source/LevelSequenceEditor/Private/Misc/MovieSceneSequenceEditor_LevelSequence.h b/Engine/Plugins/MovieScene/LevelSequenceEditor/Source/LevelSequenceEditor/Private/Misc/MovieSceneSequenceEditor_LevelSequence.h
--- a/Engine/Plugins/MovieScene/LevelSequenceEditor/Source/LevelSequenceEditor/Private/Misc/MovieSceneSequenceEditor_LevelSequence.h
+++ b/Engine/Plugins/MovieScene/LevelSequenceEditor/Source/LevelSequenceEditor/Private/Misc/MovieSceneSequenceEditor_LevelSequence.h
@@ -34,6 +34,7 @@ struct FMovieSceneSequenceEditor_LevelSequence : FMovieSceneSequenceEditor
 
 		FName BlueprintName = "SequenceDirector";
 		Blueprint = FKismetEditorUtilities::CreateBlueprint(ULevelSequenceDirector::StaticClass(), InSequence, BlueprintName, BPTYPE_Normal, UBlueprint::StaticClass(), UBlueprintGeneratedClass::StaticClass());
+		Blueprint->ClearFlags(RF_Standalone);
 
 		LevelSequence->SetDirectorBlueprint(Blueprint);
 		return Blueprint;
