#UE4 Implement GetReferencerName on some FGCObject-derived classes

diff --git a/Engine/Plugins/FX/Niagara/Source/NiagaraEditor/Public/NiagaraEditorModule.h b/Engine/Plugins/FX/Niagara/Source/NiagaraEditor/Public/NiagaraEditorModule.h
--- a/Engine/Plugins/FX/Niagara/Source/NiagaraEditor/Public/NiagaraEditorModule.h
+++ b/Engine/Plugins/FX/Niagara/Source/NiagaraEditor/Public/NiagaraEditorModule.h
@@ -122,6 +122,10 @@ private:
 
 	/** FGCObject interface */
 	virtual void AddReferencedObjects( FReferenceCollector& Collector ) override;
+	virtual FString GetReferencerName() const override
+	{
+		return "FNiagaraEditorModule";
+	}
 
 	void TestCompileScriptFromConsole(const TArray<FString>& Arguments);
 
diff --git a/Engine/Plugins/MovieScene/LevelSequenceEditor/Source/LevelSequenceEditor/Private/LevelSequenceEditorModule.cpp b/Engine/Plugins/MovieScene/LevelSequenceEditor/Source/LevelSequenceEditor/Private/LevelSequenceEditorModule.cpp
--- a/Engine/Plugins/MovieScene/LevelSequenceEditor/Source/LevelSequenceEditor/Private/LevelSequenceEditorModule.cpp
+++ b/Engine/Plugins/MovieScene/LevelSequenceEditor/Source/LevelSequenceEditor/Private/LevelSequenceEditorModule.cpp
@@ -384,6 +384,11 @@ protected:
 		}
 	}
 
+	virtual FString GetReferencerName() const override
+	{
+		return "FLevelSequenceEditorModule";
+	}
+
 private:
 
 	/** The type actions for interacting with level sequences. */
diff --git a/Engine/Plugins/Online/OnlineFramework/Source/Party/Public/Party/PartyDataReplicator.h b/Engine/Plugins/Online/OnlineFramework/Source/Party/Public/Party/PartyDataReplicator.h
--- a/Engine/Plugins/Online/OnlineFramework/Source/Party/Public/Party/PartyDataReplicator.h
+++ b/Engine/Plugins/Online/OnlineFramework/Source/Party/Public/Party/PartyDataReplicator.h
@@ -36,6 +36,11 @@ public:
 		Collector.AddReferencedObject(RepDataType);
 	}
 
+	virtual FString GetReferencerName() const override
+	{
+		return "TPartyDataReplicator";
+	}
+
 	bool IsValid() const { return RepDataType && RepDataPtr && RepDataCopy; }
 	
 	template <typename ChildRepDataT>
diff --git a/Engine/Plugins/VirtualProduction/Takes/Source/TakeRecorder/Private/TakeRecorderModule.h b/Engine/Plugins/VirtualProduction/Takes/Source/TakeRecorder/Private/TakeRecorderModule.h
--- a/Engine/Plugins/VirtualProduction/Takes/Source/TakeRecorder/Private/TakeRecorderModule.h
+++ b/Engine/Plugins/VirtualProduction/Takes/Source/TakeRecorder/Private/TakeRecorderModule.h
@@ -32,6 +32,10 @@ private:
 
 	/** FGCObject interface */
 	virtual void AddReferencedObjects(FReferenceCollector& Collector) override;
+	virtual FString GetReferencerName() const override
+	{
+		return "FTakeRecorderModule";
+	}
 
 private:
 
diff --git a/Engine/Source/Editor/MeshPaint/Public/IMeshPaintMode.h b/Engine/Source/Editor/MeshPaint/Public/IMeshPaintMode.h
--- a/Engine/Source/Editor/MeshPaint/Public/IMeshPaintMode.h
+++ b/Engine/Source/Editor/MeshPaint/Public/IMeshPaintMode.h
@@ -31,6 +31,10 @@ public:
 
 	/** FGCObject interface */
 	virtual void AddReferencedObjects( FReferenceCollector& Collector ) override;
+	virtual FString GetReferencerName() const override
+	{
+		return "IMeshPaintEdMode";
+	}
 
 	// FEdMode interface
 	virtual bool UsesToolkits() const override { return true; }
diff --git a/Engine/Source/Editor/UMGEditor/Private/UMGEditorModule.cpp b/Engine/Source/Editor/UMGEditor/Private/UMGEditorModule.cpp
--- a/Engine/Source/Editor/UMGEditor/Private/UMGEditorModule.cpp
+++ b/Engine/Source/Editor/UMGEditor/Private/UMGEditorModule.cpp
@@ -173,6 +173,11 @@ public:
 		}
 	}
 
+	virtual FString GetReferencerName() const override
+	{
+		return "FUMGEditorModule";
+	}
+
 	virtual FWidgetBlueprintCompiler* GetRegisteredCompiler() override
 	{
 		return &WidgetBlueprintCompiler;
diff --git a/Engine/Source/Editor/UnrealEd/Private/ComponentTypeRegistry.cpp b/Engine/Source/Editor/UnrealEd/Private/ComponentTypeRegistry.cpp
--- a/Engine/Source/Editor/UnrealEd/Private/ComponentTypeRegistry.cpp
+++ b/Engine/Source/Editor/UnrealEd/Private/ComponentTypeRegistry.cpp
@@ -44,6 +44,10 @@ struct FComponentTypeRegistryData
 	
 	/** Implementation of FGCObject */
 	virtual void AddReferencedObjects(FReferenceCollector& Collector) override;
+	virtual FString GetReferencerName() const override
+	{
+		return "FComponentTypeRegistryData";
+	}
 	
 	// Request a refresh of the components list next frame
 	void Invalidate()
diff --git a/Engine/Source/Editor/UnrealEd/Public/EditorActorFolders.h b/Engine/Source/Editor/UnrealEd/Public/EditorActorFolders.h
--- a/Engine/Source/Editor/UnrealEd/Public/EditorActorFolders.h
+++ b/Engine/Source/Editor/UnrealEd/Public/EditorActorFolders.h
@@ -52,6 +52,10 @@ struct UNREALED_API FActorFolders : public FGCObject
 
 	// FGCObject Interface
 	virtual void AddReferencedObjects(FReferenceCollector& Collector) override;
+	virtual FString GetReferencerName() const override
+	{
+		return "FActorFolders";
+	}
 	// End FGCObject Interface
 
 	/** Check whether the singleton is valid */
diff --git a/Engine/Source/Editor/UnrealEd/Public/EditorComponents.h b/Engine/Source/Editor/UnrealEd/Public/EditorComponents.h
--- a/Engine/Source/Editor/UnrealEd/Public/EditorComponents.h
+++ b/Engine/Source/Editor/UnrealEd/Public/EditorComponents.h
@@ -28,6 +28,10 @@ public:
 	 * @param Ar	FArchive to serialize with
 	 */
 	virtual void AddReferencedObjects( FReferenceCollector& Collector ) override;
+	virtual FString GetReferencerName() const override
+	{
+		return "FGridWidget";
+	}
 
 private:
 	// ALU based
diff --git a/Engine/Source/Editor/UnrealEd/Public/UnrealWidget.h b/Engine/Source/Editor/UnrealEd/Public/UnrealWidget.h
--- a/Engine/Source/Editor/UnrealEd/Public/UnrealWidget.h
+++ b/Engine/Source/Editor/UnrealEd/Public/UnrealWidget.h
@@ -192,6 +192,10 @@ public:
 	 * @param Ar	FArchive to serialize with
 	 */
 	virtual void AddReferencedObjects( FReferenceCollector& Collector ) override;
+	virtual FString GetReferencerName() const override
+	{
+		return "FWidget";
+	}
 
 	/**
 	 * Gets the axis to draw based on the current widget mode
diff --git a/Engine/Source/Runtime/CoreUObject/Public/UObject/StrongObjectPtr.h b/Engine/Source/Runtime/CoreUObject/Public/UObject/StrongObjectPtr.h
--- a/Engine/Source/Runtime/CoreUObject/Public/UObject/StrongObjectPtr.h
+++ b/Engine/Source/Runtime/CoreUObject/Public/UObject/StrongObjectPtr.h
@@ -45,6 +45,11 @@ namespace UE4StrongObjectPtr_Private
 			Collector.AddReferencedObject(Object);
 		}
 
+		virtual FString GetReferencerName() const override
+		{
+			return "UE4StrongObjectPtr_Private::FInternalReferenceCollector";
+		}
+
 	private:
 		const volatile UObject* Object;
 	};
