Make level sequence editor spawn register public and usable by other plugins.

diff --git a/Engine/Plugins/MovieScene/LevelSequenceEditor/Source/LevelSequenceEditor/Private/Misc/LevelSequenceEditorSpawnRegister.h b/Engine/Plugins/MovieScene/LevelSequenceEditor/Source/LevelSequenceEditor/Private/Misc/LevelSequenceEditorSpawnRegister.h
deleted file mode 100644
diff --git a/Engine/Plugins/MovieScene/LevelSequenceEditor/Source/LevelSequenceEditor/Public/Misc/LevelSequenceEditorSpawnRegister.h b/Engine/Plugins/MovieScene/LevelSequenceEditor/Source/LevelSequenceEditor/Public/Misc/LevelSequenceEditorSpawnRegister.h
new file mode 100644
--- /dev/null
+++ b/Engine/Plugins/MovieScene/LevelSequenceEditor/Source/LevelSequenceEditor/Public/Misc/LevelSequenceEditorSpawnRegister.h
@@ -0,0 +1,107 @@
+// Copyright 1998-2019 Epic Games, Inc. All Rights Reserved.
+
+#pragma once
+
+#include "CoreMinimal.h"
+#include "Templates/ValueOrError.h"
+#include "MovieSceneSequenceID.h"
+#include "MovieSceneSpawnRegister.h"
+#include "LevelSequenceSpawnRegister.h"
+#include "UObject/ObjectKey.h"
+
+class IMovieScenePlayer;
+class ISequencer;
+class UMovieScene;
+
+/**
+ * Spawn register used in the editor to add some usability features like maintaining selection states, and projecting spawned state onto spawnable defaults
+ */
+class LEVELSEQUENCEEDITOR_API FLevelSequenceEditorSpawnRegister
+	: public FLevelSequenceSpawnRegister
+{
+public:
+
+	/** Constructor */
+	FLevelSequenceEditorSpawnRegister();
+
+	/** Destructor. */
+	~FLevelSequenceEditorSpawnRegister();
+
+public:
+
+	void SetSequencer(const TSharedPtr<ISequencer>& Sequencer);
+
+public:
+
+	// FLevelSequenceSpawnRegister interface
+
+	virtual UObject* SpawnObject(FMovieSceneSpawnable& Spawnable, FMovieSceneSequenceIDRef TemplateID, IMovieScenePlayer& Player) override;
+	virtual void PreDestroyObject(UObject& Object, const FGuid& BindingId, FMovieSceneSequenceIDRef TemplateID) override;
+	virtual void SaveDefaultSpawnableState(FMovieSceneSpawnable& Spawnable, FMovieSceneSequenceIDRef TemplateID, IMovieScenePlayer& Player) override;
+#if WITH_EDITOR
+	virtual TValueOrError<FNewSpawnable, FText> CreateNewSpawnableType(UObject& SourceObject, UMovieScene& OwnerMovieScene, UActorFactory* ActorFactory = nullptr) override;
+	virtual void SetupDefaultsForSpawnable(UObject* SpawnedObject, const FGuid& Guid, const TOptional<FTransformData>& TransformData, TSharedRef<ISequencer> Sequencer, USequencerSettings* Settings) override;
+	virtual void HandleConvertPossessableToSpawnable(UObject* OldObject, IMovieScenePlayer& Player, TOptional<FTransformData>& OutTransformData) override;
+	virtual bool CanConvertSpawnableToPossessable(FMovieSceneSpawnable& Spawnable) const override;
+#endif
+
+private:
+
+	/** Called when the editor selection has changed. */
+	void HandleActorSelectionChanged(const TArray<UObject*>& NewSelection, bool bForceRefresh);
+
+	/** Saves the default state for the specified spawnable, if an instance for it currently exists */
+	void SaveDefaultSpawnableState(const FGuid& BindingId, FMovieSceneSequenceIDRef TemplateID);
+	void SaveDefaultSpawnableStateImpl(FMovieSceneSpawnable& Spawnable, UMovieSceneSequence* Sequence, UObject* SpawnedObject, IMovieScenePlayer& Player);
+
+	/** Called from the editor when a blueprint object replacement has occurred */
+	void OnObjectsReplaced(const TMap<UObject*, UObject*>& OldToNewInstanceMap);
+
+	/** Called whenever an object is modified in the editor */
+	void OnObjectModified(UObject* ModifiedObject);
+
+	/** Called before an object is saved in the editor */
+	void OnPreObjectSaved(UObject* Object);
+
+private:
+
+	struct FTrackedObjectState
+	{
+		FTrackedObjectState(FMovieSceneSequenceIDRef InTemplateID, const FGuid& InObjectBindingID) : TemplateID(InTemplateID), ObjectBindingID(InObjectBindingID), bHasBeenModified(false) {}
+
+		/** The sequence ID that spawned this object */
+		FMovieSceneSequenceID TemplateID;
+
+		/** The object binding ID of the object in the template */
+		FGuid ObjectBindingID;
+
+		/** true if this object has been modified since it was spawned and is different from the current object template */
+		bool bHasBeenModified;
+	};
+
+private:
+
+	/** Handles for delegates that we've bound to. */
+	FDelegateHandle OnActorSelectionChangedHandle;
+
+	/** Set of spawn register keys for objects that should be selected if they are spawned. */
+	TSet<FMovieSceneSpawnRegisterKey> SelectedSpawnedObjects;
+
+	/** Map from a sequenceID to an array of objects that have been tracked */
+	TMap<FObjectKey, FTrackedObjectState> TrackedObjects;
+
+	/** Set of UMovieSceneSequences that this register has spawned objects for that are modified */
+	TSet<FObjectKey> SequencesWithModifiedObjects;
+
+	/** True if we should clear the above selection cache when the editor selection has been changed. */
+	bool bShouldClearSelectionCache;
+
+	/** Weak pointer to the active sequencer. */
+	TWeakPtr<ISequencer> WeakSequencer;
+
+	/** Handle to a delegate that is bound to FCoreUObjectDelegates::OnObjectModified to harvest changes to spawned objects. */
+	FDelegateHandle OnObjectModifiedHandle;
+
+	/** Handle to a delegate that is bound to FCoreUObjectDelegates::OnObjectSaved to harvest changes to spawned objects. */
+	FDelegateHandle OnObjectSavedHandle;
+};
