Sequencer: Fix crash when adding a spawnable from a blueprint generated class

diff --git a/Engine/Plugins/MovieScene/LevelSequenceEditor/Source/LevelSequenceEditor/Private/Misc/LevelSequenceEditorActorSpawner.cpp b/Engine/Plugins/MovieScene/LevelSequenceEditor/Source/LevelSequenceEditor/Private/Misc/LevelSequenceEditorActorSpawner.cpp
--- a/Engine/Plugins/MovieScene/LevelSequenceEditor/Source/LevelSequenceEditor/Private/Misc/LevelSequenceEditorActorSpawner.cpp
+++ b/Engine/Plugins/MovieScene/LevelSequenceEditor/Source/LevelSequenceEditor/Private/Misc/LevelSequenceEditorActorSpawner.cpp
@@ -69,6 +69,20 @@ TValueOrError<FNewSpawnable, FText> FLevelSequenceEditorActorSpawner::CreateNewS
 		NewSpawnable.ObjectTemplate = NewObject<UObject>(&OwnerMovieScene, SourceBlueprint->GeneratedClass, TemplateName, RF_Transactional);
 	}
 
+	else if (UBlueprintGeneratedClass* SourceBlueprintGeneratedClass = Cast<UBlueprintGeneratedClass>(&SourceObject))
+	{
+		if (UBlueprint* BlueprintGeneratedBy = Cast<UBlueprint>(SourceBlueprintGeneratedClass->ClassGeneratedBy))
+		{
+			if (!OwnerMovieScene.GetClass()->IsChildOf(BlueprintGeneratedBy->GeneratedClass->ClassWithin))
+			{
+				ErrorText = FText::Format(LOCTEXT("ClassWithin", "Unable to add spawnable for class of type '{0}' since it has a required outer class '{1}'."), FText::FromString(SourceObject.GetName()), FText::FromString(BlueprintGeneratedBy->GeneratedClass->ClassWithin->GetName()));
+				return MakeError(ErrorText);
+			}
+
+			NewSpawnable.ObjectTemplate = NewObject<UObject>(&OwnerMovieScene, BlueprintGeneratedBy->GeneratedClass, TemplateName, RF_Transactional);
+		}
+	}
+
 	// At this point we have to assume it's an asset
 	else
 	{
