Mark LevelSequenceFactoryNew as BlueprintType so editor utlity blueprints/widgets can create Level Sequence Assets

diff --git a/Engine/Plugins/MovieScene/LevelSequenceEditor/Source/LevelSequenceEditor/Private/Factories/LevelSequenceFactoryNew.h b/Engine/Plugins/MovieScene/LevelSequenceEditor/Source/LevelSequenceEditor/Private/Factories/LevelSequenceFactoryNew.h
--- a/Engine/Plugins/MovieScene/LevelSequenceEditor/Source/LevelSequenceEditor/Private/Factories/LevelSequenceFactoryNew.h
+++ b/Engine/Plugins/MovieScene/LevelSequenceEditor/Source/LevelSequenceEditor/Private/Factories/LevelSequenceFactoryNew.h
@@ -10,7 +10,7 @@
 /**
  * Implements a factory for ULevelSequence objects.
  */
-UCLASS(hidecategories=Object)
+UCLASS(BlueprintType, hidecategories=Object)
 class ULevelSequenceFactoryNew
 	: public UFactory
 {
