Sequencer: Fix crash if get keys from selection is given a large threshold

diff --git a/Engine/Plugins/MovieScene/LevelSequenceEditor/Source/LevelSequenceEditor/Private/CinematicViewport/SCinematicTransportRange.cpp b/Engine/Plugins/MovieScene/LevelSequenceEditor/Source/LevelSequenceEditor/Private/CinematicViewport/SCinematicTransportRange.cpp
--- a/Engine/Plugins/MovieScene/LevelSequenceEditor/Source/LevelSequenceEditor/Private/CinematicViewport/SCinematicTransportRange.cpp
+++ b/Engine/Plugins/MovieScene/LevelSequenceEditor/Source/LevelSequenceEditor/Private/CinematicViewport/SCinematicTransportRange.cpp
@@ -102,7 +102,7 @@ void SCinematicTransportRange::Tick(const FGeometry& AllottedGeometry, const dou
 		FTimeToPixel TimeToPixelConverter(AllottedGeometry, WorkingRange, Sequencer->GetFocusedTickResolution());
 		const float DuplicateThreshold = (TimeToPixelConverter.PixelToSeconds(3.f) - TimeToPixelConverter.PixelToSeconds(0.f));
 
-		Sequencer->GetKeysFromSelection(ActiveKeyCollection, DuplicateThreshold);
+		Sequencer->GetKeysFromSelection(ActiveKeyCollection, FMath::Max(DuplicateThreshold, SMALL_NUMBER));
 	}
 }
 
