Sequencer: Mark the UI dirty so that editors that use the spawnable template's components can regenerate (ie. TransformTrackEditor). This fixes a bug where the camera button in the object binding track would not appear for spawned blueprints with camera components

diff --git a/Engine/Plugins/MovieScene/LevelSequenceEditor/Source/LevelSequenceEditor/Private/Misc/LevelSequenceEditorSpawnRegister.cpp b/Engine/Plugins/MovieScene/LevelSequenceEditor/Source/LevelSequenceEditor/Private/Misc/LevelSequenceEditorSpawnRegister.cpp
--- a/Engine/Plugins/MovieScene/LevelSequenceEditor/Source/LevelSequenceEditor/Private/Misc/LevelSequenceEditorSpawnRegister.cpp
+++ b/Engine/Plugins/MovieScene/LevelSequenceEditor/Source/LevelSequenceEditor/Private/Misc/LevelSequenceEditorSpawnRegister.cpp
@@ -82,6 +82,12 @@ UObject* FLevelSequenceEditorSpawnRegister::SpawnObject(FMovieSceneSpawnable& Sp
 		}
 	}
 
+	// Mark the UI dirty so that editors that use the spawnable template's components can regenerate (ie. TransformTrackEditor)
+	if (WeakSequencer.IsValid())
+	{
+		WeakSequencer.Pin()->NotifyMovieSceneDataChanged(EMovieSceneDataChangeType::MovieSceneStructureItemsChanged);
+	}
+
 	return NewObject;
 }
 
