Sequencer: Back out changelist 11644900. It's causing more UI refreshing than desirable

#rnx

diff --git a/Engine/Plugins/MovieScene/LevelSequenceEditor/Source/LevelSequenceEditor/Private/Misc/LevelSequenceEditorSpawnRegister.cpp b/Engine/Plugins/MovieScene/LevelSequenceEditor/Source/LevelSequenceEditor/Private/Misc/LevelSequenceEditorSpawnRegister.cpp
--- a/Engine/Plugins/MovieScene/LevelSequenceEditor/Source/LevelSequenceEditor/Private/Misc/LevelSequenceEditorSpawnRegister.cpp
+++ b/Engine/Plugins/MovieScene/LevelSequenceEditor/Source/LevelSequenceEditor/Private/Misc/LevelSequenceEditorSpawnRegister.cpp
@@ -82,12 +82,6 @@ UObject* FLevelSequenceEditorSpawnRegister::SpawnObject(FMovieSceneSpawnable& Sp
 		}
 	}
 
-	// Mark the UI dirty so that editors that use the spawnable template's components can regenerate (ie. TransformTrackEditor)
-	if (WeakSequencer.IsValid())
-	{
-		WeakSequencer.Pin()->NotifyMovieSceneDataChanged(EMovieSceneDataChangeType::MovieSceneStructureItemsChanged);
-	}
-
 	return NewObject;
 }
 
