Sequencer: Fix crash creating actor with an exsting name. Replaced GWorld with GCurrentLevelEditingViewportClient

diff --git a/Engine/Plugins/MovieScene/LevelSequenceEditor/Source/LevelSequenceEditor/Private/Misc/LevelSequenceEditorActorSpawner.cpp b/Engine/Plugins/MovieScene/LevelSequenceEditor/Source/LevelSequenceEditor/Private/Misc/LevelSequenceEditorActorSpawner.cpp
--- a/Engine/Plugins/MovieScene/LevelSequenceEditor/Source/LevelSequenceEditor/Private/Misc/LevelSequenceEditorActorSpawner.cpp
+++ b/Engine/Plugins/MovieScene/LevelSequenceEditor/Source/LevelSequenceEditor/Private/Misc/LevelSequenceEditorActorSpawner.cpp
@@ -106,13 +106,19 @@ TValueOrError<FNewSpawnable, FText> FLevelSequenceEditorActorSpawner::CreateNewS
 				}
 			}
 
-			AActor* Instance = FactoryToUse->CreateActor(&SourceObject, GWorld->PersistentLevel, FTransform(), RF_Transient | RF_Transactional, TemplateName );
-			Instance->bIsEditorPreviewActor = false;
-			NewSpawnable.ObjectTemplate = StaticDuplicateObject(Instance, &OwnerMovieScene, TemplateName, RF_AllFlags & ~RF_Transient);
+			UWorld* World = GCurrentLevelEditingViewportClient ? GCurrentLevelEditingViewportClient->GetWorld() : nullptr;
+			if (World)
+			{
+				const FName ActorName = MakeUniqueObjectName(World->PersistentLevel, FactoryToUse->NewActorClass->StaticClass(), TemplateName);
+
+				AActor* Instance = FactoryToUse->CreateActor(&SourceObject, World->PersistentLevel, FTransform(), RF_Transient | RF_Transactional, ActorName );
+				Instance->bIsEditorPreviewActor = false;
+				NewSpawnable.ObjectTemplate = StaticDuplicateObject(Instance, &OwnerMovieScene, TemplateName, RF_AllFlags & ~RF_Transient);
 
-			const bool bNetForce = false;
-			const bool bShouldModifyLevel = false;
-			GWorld->DestroyActor(Instance, bNetForce, bShouldModifyLevel);
+				const bool bNetForce = false;
+				const bool bShouldModifyLevel = false;
+				World->DestroyActor(Instance, bNetForce, bShouldModifyLevel);
+			}
 		}
 	}
 
