MoviePipeline: Sequencer no longer auto-binds playback contexts to the world created during movie rendering. The world is not available to the executor when this function is called so this system cannot return the world like the previous one, instead it passes a mutable bool around through the multicast delegate. The specific world was not used anyways as all it was trying to do was calculate if it should return the editor world for this function (during movie renders).

diff --git a/Engine/Plugins/MovieScene/LevelSequenceEditor/Source/LevelSequenceEditor/Private/LevelSequenceEditorModule.cpp b/Engine/Plugins/MovieScene/LevelSequenceEditor/Source/LevelSequenceEditor/Private/LevelSequenceEditorModule.cpp
--- a/Engine/Plugins/MovieScene/LevelSequenceEditor/Source/LevelSequenceEditor/Private/LevelSequenceEditorModule.cpp
+++ b/Engine/Plugins/MovieScene/LevelSequenceEditor/Source/LevelSequenceEditor/Private/LevelSequenceEditorModule.cpp
@@ -370,6 +370,11 @@ protected:
 		return OnMasterSequenceCreatedEvent;
 	}
 
+	FAllowPlaybackContext& OnComputePlaybackContext() override
+	{
+		return OnComputePlaybackContextDelegate;
+	}
+
 	static TSharedRef<ISequencerEditorObjectBinding> OnCreateActorBinding(TSharedRef<ISequencer> InSequencer)
 	{
 		return MakeShareable(new FLevelSequenceEditorActorBinding(InSequencer));
@@ -401,6 +406,8 @@ private:
 
 	FOnMasterSequenceCreated OnMasterSequenceCreatedEvent;
 
+	FAllowPlaybackContext OnComputePlaybackContextDelegate;
+
 	FDelegateHandle ActorBindingDelegateHandle;
 
 	FDelegateHandle EditorActorSpawnerDelegateHandle;
diff --git a/Engine/Plugins/MovieScene/LevelSequenceEditor/Source/LevelSequenceEditor/Private/Misc/LevelSequencePlaybackContext.cpp b/Engine/Plugins/MovieScene/LevelSequenceEditor/Source/LevelSequenceEditor/Private/Misc/LevelSequencePlaybackContext.cpp
--- a/Engine/Plugins/MovieScene/LevelSequenceEditor/Source/LevelSequenceEditor/Private/Misc/LevelSequencePlaybackContext.cpp
+++ b/Engine/Plugins/MovieScene/LevelSequenceEditor/Source/LevelSequenceEditor/Private/Misc/LevelSequencePlaybackContext.cpp
@@ -17,6 +17,7 @@
 #include "Framework/MultiBox/MultiBoxBuilder.h"
 
 #include "MovieSceneCaptureDialogModule.h"
+#include "LevelSequenceEditorModule.h"
 #include "EditorStyleSet.h"
 #include "Engine/Engine.h"
 #include "Editor.h"
@@ -191,6 +192,15 @@ UWorld* FLevelSequencePlaybackContext::ComputePlaybackContext()
 	const ULevelSequenceEditorSettings* Settings            = GetDefault<ULevelSequenceEditorSettings>();
 	IMovieSceneCaptureDialogModule*     CaptureDialogModule = FModuleManager::GetModulePtr<IMovieSceneCaptureDialogModule>("MovieSceneCaptureDialog");
 
+	// Some plugins may not want us to automatically attempt to bind to the world where it doesn't make sense,
+	// such as movie rendering.
+	bool bAllowPlaybackContextBinding = true;
+	ILevelSequenceEditorModule* LevelSequenceEditorModule = FModuleManager::GetModulePtr<ILevelSequenceEditorModule>("LevelSequenceEditor");
+	if (LevelSequenceEditorModule)
+	{
+		LevelSequenceEditorModule->OnComputePlaybackContext().Broadcast(bAllowPlaybackContextBinding);
+	}
+
 	UWorld* RecordingWorld = CaptureDialogModule ? CaptureDialogModule->GetCurrentlyRecordingWorld() : nullptr;
 
 	// Only allow PIE and Simulate worlds if the settings allow them
@@ -205,7 +215,7 @@ UWorld* FLevelSequencePlaybackContext::ComputePlaybackContext()
 		if (Context.WorldType == EWorldType::PIE)
 		{
 			UWorld* ThisWorld = Context.World();
-			if (bIsPIEValid && RecordingWorld != ThisWorld)
+			if (bIsPIEValid && bAllowPlaybackContextBinding && RecordingWorld != ThisWorld)
 			{
 				return ThisWorld;
 			}
@@ -343,4 +353,4 @@ TSharedRef<SWidget> SLevelSequenceContextPicker::BuildWorldPickerMenu()
 	return MenuBuilder.MakeWidget();
 }
 
-#undef LOCTEXT_NAMESPACE
\ No newline at end of file
+#undef LOCTEXT_NAMESPACE
diff --git a/Engine/Plugins/MovieScene/LevelSequenceEditor/Source/LevelSequenceEditor/Public/LevelSequenceEditorModule.h b/Engine/Plugins/MovieScene/LevelSequenceEditor/Source/LevelSequenceEditor/Public/LevelSequenceEditorModule.h
--- a/Engine/Plugins/MovieScene/LevelSequenceEditor/Source/LevelSequenceEditor/Public/LevelSequenceEditorModule.h
+++ b/Engine/Plugins/MovieScene/LevelSequenceEditor/Source/LevelSequenceEditor/Public/LevelSequenceEditorModule.h
@@ -18,4 +18,7 @@ class ILevelSequenceEditorModule
 public:
 	DECLARE_EVENT_OneParam(ILevelSequenceEditorModule, FOnMasterSequenceCreated, UObject*);
 	virtual FOnMasterSequenceCreated& OnMasterSequenceCreated() = 0;
+
+	DECLARE_MULTICAST_DELEGATE_OneParam(FAllowPlaybackContext, bool&);
+	virtual FAllowPlaybackContext& OnComputePlaybackContext() = 0;
 };
diff --git a/Engine/Plugins/MovieScene/MovieRenderPipeline/MovieRenderPipeline.uplugin b/Engine/Plugins/MovieScene/MovieRenderPipeline/MovieRenderPipeline.uplugin
--- a/Engine/Plugins/MovieScene/MovieRenderPipeline/MovieRenderPipeline.uplugin
+++ b/Engine/Plugins/MovieScene/MovieRenderPipeline/MovieRenderPipeline.uplugin
@@ -18,6 +18,11 @@
 		{
 			"Name": "EditorScriptingUtilities",
 			"Enabled": true
+		},
+		{
+			"Name": "LevelSequenceEditor",
+			"Enabled": true,
+			"Optional": true
 		}
 	],
 	"Modules": [
diff --git a/Engine/Plugins/MovieScene/MovieRenderPipeline/Source/MovieRenderPipelineEditor/MovieRenderPipelineEditor.Build.cs b/Engine/Plugins/MovieScene/MovieRenderPipeline/Source/MovieRenderPipelineEditor/MovieRenderPipelineEditor.Build.cs
--- a/Engine/Plugins/MovieScene/MovieRenderPipeline/Source/MovieRenderPipelineEditor/MovieRenderPipelineEditor.Build.cs
+++ b/Engine/Plugins/MovieScene/MovieRenderPipeline/Source/MovieRenderPipelineEditor/MovieRenderPipelineEditor.Build.cs
@@ -51,7 +51,8 @@ public class MovieRenderPipelineEditor : ModuleRules
                 "LevelSequence",
                 "UnrealEd",
                 "WorkspaceMenuStructure",
-				"EditorStyle"
+				"EditorStyle",
+				"LevelSequenceEditor"
             }
         );
 
diff --git a/Engine/Plugins/MovieScene/MovieRenderPipeline/Source/MovieRenderPipelineEditor/Private/MoviePipelineQueueSubsystem.cpp b/Engine/Plugins/MovieScene/MovieRenderPipeline/Source/MovieRenderPipelineEditor/Private/MoviePipelineQueueSubsystem.cpp
--- a/Engine/Plugins/MovieScene/MovieRenderPipeline/Source/MovieRenderPipelineEditor/Private/MoviePipelineQueueSubsystem.cpp
+++ b/Engine/Plugins/MovieScene/MovieRenderPipeline/Source/MovieRenderPipelineEditor/Private/MoviePipelineQueueSubsystem.cpp
@@ -1,6 +1,8 @@
 // Copyright Epic Games, Inc. All Rights Reserved.
 
 #include "MoviePipelineQueueSubsystem.h"
+#include "LevelSequenceEditorModule.h"
+#include "Modules/ModuleManager.h"
 
 UMoviePipelineExecutorBase* UMoviePipelineQueueSubsystem::RenderQueueWithExecutor(TSubclassOf<UMoviePipelineExecutorBase> InExecutorType)
 {
@@ -8,16 +10,35 @@ UMoviePipelineExecutorBase* UMoviePipelineQueueSubsystem::RenderQueueWithExecuto
 	{
 		return nullptr;
 	}
-	
+
+	ILevelSequenceEditorModule* LevelSequenceEditorModule = FModuleManager::GetModulePtr<ILevelSequenceEditorModule>("LevelSequenceEditor");
+	if (LevelSequenceEditorModule)
+	{
+		LevelSequenceEditorModule->OnComputePlaybackContext().AddUObject(this, &UMoviePipelineQueueSubsystem::OnSequencerContextBinding);
+	}
+		
 	ActiveExecutor = NewObject<UMoviePipelineExecutorBase>(this, InExecutorType);
 	ActiveExecutor->OnExecutorFinished().AddUObject(this, &UMoviePipelineQueueSubsystem::OnExecutorFinished);
 	ActiveExecutor->Execute(GetQueue());
-	
 	return ActiveExecutor;
 }
 
+void UMoviePipelineQueueSubsystem::OnSequencerContextBinding(bool& bAllowBinding)
+{
+	if (ActiveExecutor && ActiveExecutor->IsRendering())
+	{
+		bAllowBinding = false;
+	}
+}
+
 void UMoviePipelineQueueSubsystem::OnExecutorFinished(UMoviePipelineExecutorBase* InPipelineExecutor, bool bSuccess)
 {
+	ILevelSequenceEditorModule* LevelSequenceEditorModule = FModuleManager::GetModulePtr<ILevelSequenceEditorModule>("LevelSequenceEditorModule");
+	if (LevelSequenceEditorModule)
+	{
+		LevelSequenceEditorModule->OnComputePlaybackContext().RemoveAll(this);
+	}
+
 	ActiveExecutor = nullptr;
 }
 
diff --git a/Engine/Plugins/MovieScene/MovieRenderPipeline/Source/MovieRenderPipelineEditor/Public/MoviePipelineQueueSubsystem.h b/Engine/Plugins/MovieScene/MovieRenderPipeline/Source/MovieRenderPipelineEditor/Public/MoviePipelineQueueSubsystem.h
--- a/Engine/Plugins/MovieScene/MovieRenderPipeline/Source/MovieRenderPipelineEditor/Public/MoviePipelineQueueSubsystem.h
+++ b/Engine/Plugins/MovieScene/MovieRenderPipeline/Source/MovieRenderPipelineEditor/Public/MoviePipelineQueueSubsystem.h
@@ -58,6 +58,7 @@ public:
 private:
 	/** Called when the executor is finished so that we can release it and stop reporting IsRendering() == true. */
 	void OnExecutorFinished(UMoviePipelineExecutorBase* InPipelineExecutor, bool bSuccess);
+	void OnSequencerContextBinding(bool& bAllowBinding);
 
 private:
 	UPROPERTY(Transient)
