Fix missing background on some elements in sequencer

diff --git a/Engine/Plugins/MovieScene/LevelSequenceEditor/Source/LevelSequenceEditor/Private/Misc/LevelSequencePlaybackContext.cpp b/Engine/Plugins/MovieScene/LevelSequenceEditor/Source/LevelSequenceEditor/Private/Misc/LevelSequencePlaybackContext.cpp
--- a/Engine/Plugins/MovieScene/LevelSequenceEditor/Source/LevelSequenceEditor/Private/Misc/LevelSequencePlaybackContext.cpp
+++ b/Engine/Plugins/MovieScene/LevelSequenceEditor/Source/LevelSequenceEditor/Private/Misc/LevelSequencePlaybackContext.cpp
@@ -66,7 +66,7 @@ private:
 		}
 		else
 		{
-			return FEditorStyle::GetBrush("LevelViewport.NoViewportBorder");
+			return FStyleDefaults::GetNoBrush();
 		}
 	}
 
