Sequencer Scripting: Expose camera cut locking to scripting

IsCameraCutLockedToViewport
SetLockCameraCutToViewport

diff --git a/Engine/Plugins/MovieScene/LevelSequenceEditor/Source/LevelSequenceEditor/Private/LevelSequenceEditorBlueprintLibrary.cpp b/Engine/Plugins/MovieScene/LevelSequenceEditor/Source/LevelSequenceEditor/Private/LevelSequenceEditorBlueprintLibrary.cpp
--- a/Engine/Plugins/MovieScene/LevelSequenceEditor/Source/LevelSequenceEditor/Private/LevelSequenceEditorBlueprintLibrary.cpp
+++ b/Engine/Plugins/MovieScene/LevelSequenceEditor/Source/LevelSequenceEditor/Private/LevelSequenceEditorBlueprintLibrary.cpp
@@ -8,6 +8,7 @@
 
 #include "Modules/ModuleManager.h"
 #include "LevelEditor.h"
+#include "LevelEditorViewport.h"
 #include "Subsystems/AssetEditorSubsystem.h"
 
 #include "MovieSceneCommonHelpers.h"
@@ -328,3 +329,44 @@ void ULevelSequenceEditorBlueprintLibrary::SetLockLevelSequence(bool bLock)
 	}
 }
 
+bool ULevelSequenceEditorBlueprintLibrary::IsCameraCutLockedToViewport()
+{
+	if (CurrentSequencer.IsValid())
+	{
+		TSharedPtr<ISequencer> Sequencer = CurrentSequencer.Pin();
+		return Sequencer->IsPerspectiveViewportCameraCutEnabled();
+	}
+
+	return false;
+}
+
+void ULevelSequenceEditorBlueprintLibrary::SetLockCameraCutToViewport(bool bLock)
+{
+	if (CurrentSequencer.IsValid())
+	{
+		TSharedPtr<ISequencer> Sequencer = CurrentSequencer.Pin();
+
+		if (bLock)
+		{
+			for(FLevelEditorViewportClient* LevelVC : GEditor->GetLevelViewportClients())
+			{
+				if (LevelVC && LevelVC->AllowsCinematicControl() && LevelVC->GetViewMode() != VMI_Unknown)
+				{
+					LevelVC->SetActorLock(nullptr);
+					LevelVC->bLockedCameraView = false;
+					LevelVC->UpdateViewForLockedActor();
+					LevelVC->Invalidate();
+				}
+			}
+			Sequencer->SetPerspectiveViewportCameraCutEnabled(true);
+		}
+		else
+		{
+			Sequencer->UpdateCameraCut(nullptr, EMovieSceneCameraCutParams());
+			Sequencer->SetPerspectiveViewportCameraCutEnabled(false);
+		}
+
+		Sequencer->ForceEvaluate();
+	}
+}
+
diff --git a/Engine/Plugins/MovieScene/LevelSequenceEditor/Source/LevelSequenceEditor/Private/LevelSequenceEditorBlueprintLibrary.h b/Engine/Plugins/MovieScene/LevelSequenceEditor/Source/LevelSequenceEditor/Private/LevelSequenceEditorBlueprintLibrary.h
--- a/Engine/Plugins/MovieScene/LevelSequenceEditor/Source/LevelSequenceEditor/Private/LevelSequenceEditorBlueprintLibrary.h
+++ b/Engine/Plugins/MovieScene/LevelSequenceEditor/Source/LevelSequenceEditor/Private/LevelSequenceEditorBlueprintLibrary.h
@@ -163,6 +163,16 @@ public:
 	UFUNCTION(BlueprintCallable, Category = "Level Sequence Editor")
 	static void SetLockLevelSequence(bool bLock);
 
+public:
+
+	/** Check whether the lock for the viewport to the camera cuts is enabled. */
+	UFUNCTION(BlueprintPure, Category = "Level Sequence Editor")
+	static bool IsCameraCutLockedToViewport();
+
+	/** Sets the lock for the viewport to the camera cuts. */
+	UFUNCTION(BlueprintCallable, Category = "Level Sequence Editor")
+	static void SetLockCameraCutToViewport(bool bLock);
+
 public:
 
 	/*
