Sequencer: Fix convert to spawnable. The actual bug here is in SetupDefaultsForSpawnable where the default value is set for ALL existing transform sections. That seems wrong and unnecessary. Most of the time a spawnable will be created from scratch, so this problem hasn't surfaced. Also, the problem wouldn't appear if all the channels have keys because the default value wouldn't be used. However, this has appeared when converting from possessable to spawnable because the spawnable could inherit transform sections from the possessable. In this case, the existing transform sections shoudn't be overwritten. The default values should only be set if there is no existing transform section and a new one is created.

diff --git a/Engine/Plugins/MovieScene/LevelSequenceEditor/Source/LevelSequenceEditor/Private/Misc/LevelSequenceEditorActorSpawner.cpp b/Engine/Plugins/MovieScene/LevelSequenceEditor/Source/LevelSequenceEditor/Private/Misc/LevelSequenceEditorActorSpawner.cpp
--- a/Engine/Plugins/MovieScene/LevelSequenceEditor/Source/LevelSequenceEditor/Private/Misc/LevelSequenceEditorActorSpawner.cpp
+++ b/Engine/Plugins/MovieScene/LevelSequenceEditor/Source/LevelSequenceEditor/Private/Misc/LevelSequenceEditorActorSpawner.cpp
@@ -250,42 +250,34 @@ void FLevelSequenceEditorActorSpawner::SetupDefaultsForSpawnable(UObject* Spawne
 			TransformTrack = Cast<UMovieScene3DTransformTrack>(OwnerMovieScene->AddTrack(UMovieScene3DTransformTrack::StaticClass(), Guid));
 		}
 
-		if (TransformTrack)
+		if (TransformTrack && TransformTrack->GetAllSections().Num() == 0)
 		{
-			const TArray<UMovieSceneSection*>& Sections = TransformTrack->GetAllSections();
-			if (!Sections.Num())
-			{
-				TransformTrack->AddSection(*TransformTrack->CreateNewSection());
-			}
+			UMovieScene3DTransformSection* TransformSection = Cast<UMovieScene3DTransformSection>(TransformTrack->CreateNewSection());
+			TransformTrack->AddSection(*TransformSection);
 
 			FVector Location = DefaultTransform->Translation;
 			FVector Rotation = DefaultTransform->Rotation.Euler();
 			FVector Scale    = DefaultTransform->Scale;
 
-			for (UMovieSceneSection* Section : Sections)
+			// Set the section to be infinite if necessary
+			if (Sequencer->GetInfiniteKeyAreas())
 			{
-				UMovieScene3DTransformSection* TransformSection = CastChecked<UMovieScene3DTransformSection>(Section);
-
-				// Set the section to be infinite if necessary
-				if (Sequencer->GetInfiniteKeyAreas())
-				{
-					TransformSection->SetRange(TRange<FFrameNumber>::All());
-				}
+				TransformSection->SetRange(TRange<FFrameNumber>::All());
+			}
 
-				// Set the section's default values to this default transform
-				TArrayView<FMovieSceneFloatChannel*> FloatChannels = TransformSection->GetChannelProxy().GetChannels<FMovieSceneFloatChannel>();
-				FloatChannels[0]->SetDefault(Location.X);
-				FloatChannels[1]->SetDefault(Location.Y);
-				FloatChannels[2]->SetDefault(Location.Z);
+			// Set the section's default values to this default transform
+			TArrayView<FMovieSceneFloatChannel*> FloatChannels = TransformSection->GetChannelProxy().GetChannels<FMovieSceneFloatChannel>();
+			FloatChannels[0]->SetDefault(Location.X);
+			FloatChannels[1]->SetDefault(Location.Y);
+			FloatChannels[2]->SetDefault(Location.Z);
 
-				FloatChannels[3]->SetDefault(Rotation.X);
-				FloatChannels[4]->SetDefault(Rotation.Y);
-				FloatChannels[5]->SetDefault(Rotation.Z);
+			FloatChannels[3]->SetDefault(Rotation.X);
+			FloatChannels[4]->SetDefault(Rotation.Y);
+			FloatChannels[5]->SetDefault(Rotation.Z);
 
-				FloatChannels[6]->SetDefault(Scale.X);
-				FloatChannels[7]->SetDefault(Scale.Y);
-				FloatChannels[8]->SetDefault(Scale.Z);
-			}
+			FloatChannels[6]->SetDefault(Scale.X);
+			FloatChannels[7]->SetDefault(Scale.Y);
+			FloatChannels[8]->SetDefault(Scale.Z);
 		}
 	}
 }
diff --git a/Engine/Source/Runtime/LevelSequence/Private/LevelSequenceActorSpawner.cpp b/Engine/Source/Runtime/LevelSequence/Private/LevelSequenceActorSpawner.cpp
--- a/Engine/Source/Runtime/LevelSequence/Private/LevelSequenceActorSpawner.cpp
+++ b/Engine/Source/Runtime/LevelSequence/Private/LevelSequenceActorSpawner.cpp
@@ -163,6 +163,7 @@ UObject* FLevelSequenceActorSpawner::SpawnObject(FMovieSceneSpawnable& Spawnable
 	{
 		SpawnTransform.SetTranslation(RootComponent->GetRelativeLocation());
 		SpawnTransform.SetRotation(RootComponent->GetRelativeRotation().Quaternion());
+		SpawnTransform.SetScale3D(RootComponent->GetRelativeScale3D());
 	}
 	else
 	{
