Cinematic Viewport: Fix popping in the viewport when hiding the toolbar

diff --git a/Engine/Plugins/MovieScene/LevelSequenceEditor/Source/LevelSequenceEditor/Private/CinematicViewport/SCinematicLevelViewport.cpp b/Engine/Plugins/MovieScene/LevelSequenceEditor/Source/LevelSequenceEditor/Private/CinematicViewport/SCinematicLevelViewport.cpp
--- a/Engine/Plugins/MovieScene/LevelSequenceEditor/Source/LevelSequenceEditor/Private/CinematicViewport/SCinematicLevelViewport.cpp
+++ b/Engine/Plugins/MovieScene/LevelSequenceEditor/Source/LevelSequenceEditor/Private/CinematicViewport/SCinematicLevelViewport.cpp
@@ -263,18 +263,24 @@ void SCinematicLevelViewport::Construct(const FArguments& InArgs)
 			.Padding(5.f)
 			.AutoHeight()
 			[
-				SNew(SHorizontalBox)
-				.Visibility_Lambda([] { return GLevelEditorModeTools().IsViewportUIHidden() ? EVisibility::Hidden : EVisibility::Visible; })
-
-				+ SHorizontalBox::Slot()
+				SNew(SVerticalBox)
+	
+				+ SVerticalBox::Slot()
+				.MaxHeight(25.f)
 				[
-					ViewportWidget->MakeExternalViewportToolbar().ToSharedRef()
-				]
+					SNew(SHorizontalBox)
+					.Visibility_Lambda([] { return GLevelEditorModeTools().IsViewportUIHidden() ? EVisibility::Hidden : EVisibility::Visible; })
 
-				+ SHorizontalBox::Slot()
-				.AutoWidth()
-				[
-					FilmOverlayOptions					
+					+ SHorizontalBox::Slot()
+					[
+						ViewportWidget->MakeExternalViewportToolbar().ToSharedRef()
+					]
+
+					+ SHorizontalBox::Slot()
+					.AutoWidth()
+					[
+						FilmOverlayOptions					
+					]
 				]
 			]
 
