Sequencer: Nullptr check for master sequence

diff --git a/Engine/Plugins/MovieScene/LevelSequenceEditor/Source/LevelSequenceEditor/Private/LevelSequenceEditorToolkit.cpp b/Engine/Plugins/MovieScene/LevelSequenceEditor/Source/LevelSequenceEditor/Private/LevelSequenceEditorToolkit.cpp
--- a/Engine/Plugins/MovieScene/LevelSequenceEditor/Source/LevelSequenceEditor/Private/LevelSequenceEditorToolkit.cpp
+++ b/Engine/Plugins/MovieScene/LevelSequenceEditor/Source/LevelSequenceEditor/Private/LevelSequenceEditorToolkit.cpp
@@ -733,6 +733,12 @@ void FLevelSequenceEditorToolkit::AddShot(UMovieSceneCinematicShotTrack* ShotTra
 
 void FLevelSequenceEditorToolkit::HandleMasterSequenceCreated(UObject* MasterSequenceAsset)
 {
+	UMovieSceneSequence* MasterSequence = Cast<UMovieSceneSequence>(MasterSequenceAsset);
+	if (!MasterSequence)
+	{
+		return;
+	}
+
 	const FScopedTransaction Transaction( LOCTEXT( "CreateMasterSequence", "Create Master Sequence" ) );
 	
 	const ULevelSequenceMasterSequenceSettings* MasterSequenceSettings = GetDefault<ULevelSequenceMasterSequenceSettings>();
@@ -741,7 +747,6 @@ void FLevelSequenceEditorToolkit::HandleMasterSequenceCreated(UObject* MasterSeq
 
 	const UMovieSceneToolsProjectSettings* ProjectSettings = GetDefault<UMovieSceneToolsProjectSettings>();
 
-	UMovieSceneSequence* MasterSequence = Cast<UMovieSceneSequence>(MasterSequenceAsset);
 	UMovieSceneCinematicShotTrack* ShotTrack = MasterSequence->GetMovieScene()->AddMasterTrack<UMovieSceneCinematicShotTrack>();
 
 	FFrameRate TickResolution = MasterSequence->GetMovieScene()->GetTickResolution();
