Update the LevelSequenceEditorToolkit to clean up the Sequencer->LevelEditor connection on shutdown rather than on UnregisterTabSpawners.

diff --git a/Engine/Plugins/MovieScene/LevelSequenceEditor/Source/LevelSequenceEditor/Private/LevelSequenceEditorToolkit.cpp b/Engine/Plugins/MovieScene/LevelSequenceEditor/Source/LevelSequenceEditor/Private/LevelSequenceEditorToolkit.cpp
--- a/Engine/Plugins/MovieScene/LevelSequenceEditor/Source/LevelSequenceEditor/Private/LevelSequenceEditorToolkit.cpp
+++ b/Engine/Plugins/MovieScene/LevelSequenceEditor/Source/LevelSequenceEditor/Private/LevelSequenceEditorToolkit.cpp
@@ -117,17 +117,28 @@ FLevelSequenceEditorToolkit::~FLevelSequenceEditorToolkit()
 {
 	OpenToolkits.Remove(this);
 
-	FLevelEditorSequencerIntegration::Get().RemoveSequencer(Sequencer.ToSharedRef());
-
-	Sequencer->Close();
-
-	// unregister delegates
 	if (FModuleManager::Get().IsModuleLoaded(TEXT("LevelEditor")))
 	{
-		auto& LevelEditorModule = FModuleManager::LoadModuleChecked<FLevelEditorModule>(TEXT("LevelEditor"));
+		FLevelEditorModule& LevelEditorModule = FModuleManager::LoadModuleChecked<FLevelEditorModule>("LevelEditor");
+		if (!IsWorldCentricAssetEditor())
+		{
+			TSharedPtr<FTabManager> LevelEditorTabManager = LevelEditorModule.GetLevelEditorTabManager();
+			LevelEditorTabManager->UnregisterTabSpawner(SequencerMainTabId);
+		}
+
+		// @todo remove when world-centric mode is added
+		LevelEditorModule.AttachSequencer(SNullWidget::NullWidget, nullptr);
+		FLevelEditorSequencerIntegration::Get().RemoveSequencer(Sequencer.ToSharedRef());
+
+		// unregister delegates
+
 		LevelEditorModule.OnMapChanged().RemoveAll(this);
 	}
 
+	Sequencer->Close();
+
+	
+
 	if (FModuleManager::Get().IsModuleLoaded(TEXT("LevelSequenceEditor")))
 	{
 		auto& LevelSequenceEditorModule = FModuleManager::LoadModuleChecked<ILevelSequenceEditorModule>(TEXT("LevelSequenceEditor"));
@@ -218,6 +229,18 @@ void FLevelSequenceEditorToolkit::Initialize(const EToolkitMode::Type Mode, cons
 			LevelEditorTabManager->TryInvokeTab(FName("LevelEditorSceneOutliner"))->RequestCloseTab();
 			LevelEditorTabManager->TryInvokeTab(FName("LevelEditorSceneOutliner"));
 		}
+
+		// For World Centric Asset Editors this isn't called until way too late in the initialization flow
+		// (ie: when you actually start to edit an asset), so the tab will be unrecognized upon restore.
+		// Because of this, the Sequencer Tab Spawner is actually registered in SLevelEditor.cpp manually
+		// which is early enough that you can restore the tab after an editor restart.
+
+		WorkspaceMenuCategory = LevelEditorTabManager->AddLocalWorkspaceMenuCategory(LOCTEXT("WorkspaceMenu_SequencerAssetEditor", "Sequencer"));
+
+		LevelEditorTabManager->RegisterTabSpawner(SequencerMainTabId, FOnSpawnTab::CreateSP(this, &FLevelSequenceEditorToolkit::HandleTabManagerSpawnTab))
+			.SetDisplayName(LOCTEXT("SequencerMainTab", "Sequencer"))
+			.SetGroup(WorkspaceMenuCategory.ToSharedRef())
+			.SetIcon(FSlateIcon(Style->GetStyleSetName(), "LevelSequenceEditor.Tabs.Sequencer"));
 	}
 	
 	// Now Attach so this window will apear in the correct front first order
@@ -281,39 +304,6 @@ FString FLevelSequenceEditorToolkit::GetWorldCentricTabPrefix() const
 }
 
 
-void FLevelSequenceEditorToolkit::RegisterTabSpawners(const TSharedRef<class FTabManager>& InTabManager)
-{
-	// For World Centric Asset Editors this isn't called until way too late in the initialization flow
-	// (ie: when you actually start to edit an asset), so the tab will be unrecognized upon restore.
-	// Because of this, the Sequencer Tab Spawner is actually registered in SLevelEditor.cpp manually
-	// which is early enough that you can restore the tab after an editor restart.
-	if (IsWorldCentricAssetEditor())
-	{
-		return;
-	}
-
-	WorkspaceMenuCategory = InTabManager->AddLocalWorkspaceMenuCategory(LOCTEXT("WorkspaceMenu_SequencerAssetEditor", "Sequencer"));
-
-	InTabManager->RegisterTabSpawner(SequencerMainTabId, FOnSpawnTab::CreateSP(this, &FLevelSequenceEditorToolkit::HandleTabManagerSpawnTab))
-		.SetDisplayName(LOCTEXT("SequencerMainTab", "Sequencer"))
-		.SetGroup(WorkspaceMenuCategory.ToSharedRef())
-		.SetIcon(FSlateIcon(Style->GetStyleSetName(), "LevelSequenceEditor.Tabs.Sequencer"));
-}
-
-
-void FLevelSequenceEditorToolkit::UnregisterTabSpawners(const TSharedRef<class FTabManager>& InTabManager)
-{
-	if (!IsWorldCentricAssetEditor())
-	{
-		InTabManager->UnregisterTabSpawner(SequencerMainTabId);
-	}
-
-	// @todo remove when world-centric mode is added
-	FLevelEditorModule& LevelEditorModule = FModuleManager::LoadModuleChecked<FLevelEditorModule>("LevelEditor");
-	LevelEditorModule.AttachSequencer(SNullWidget::NullWidget, nullptr);
-}
-
-
 /* FLevelSequenceEditorToolkit implementation
  *****************************************************************************/
 
diff --git a/Engine/Plugins/MovieScene/LevelSequenceEditor/Source/LevelSequenceEditor/Private/LevelSequenceEditorToolkit.h b/Engine/Plugins/MovieScene/LevelSequenceEditor/Source/LevelSequenceEditor/Private/LevelSequenceEditorToolkit.h
--- a/Engine/Plugins/MovieScene/LevelSequenceEditor/Source/LevelSequenceEditor/Private/LevelSequenceEditorToolkit.h
+++ b/Engine/Plugins/MovieScene/LevelSequenceEditor/Source/LevelSequenceEditor/Private/LevelSequenceEditorToolkit.h
@@ -104,8 +104,6 @@ public:
 	virtual FName GetToolkitFName() const override;
 	virtual FLinearColor GetWorldCentricTabColorScale() const override;
 	virtual FString GetWorldCentricTabPrefix() const override;
-	virtual void RegisterTabSpawners(const TSharedRef<FTabManager>& TabManager) override;
-	virtual void UnregisterTabSpawners(const TSharedRef<FTabManager>& TabManager) override;
 
 protected:
 
