Removing redundant Level Sequence Editor spawner being added.

diff --git a/Engine/Plugins/MovieScene/LevelSequenceEditor/Source/LevelSequenceEditor/Private/LevelSequenceEditorToolkit.cpp b/Engine/Plugins/MovieScene/LevelSequenceEditor/Source/LevelSequenceEditor/Private/LevelSequenceEditorToolkit.cpp
--- a/Engine/Plugins/MovieScene/LevelSequenceEditor/Source/LevelSequenceEditor/Private/LevelSequenceEditorToolkit.cpp
+++ b/Engine/Plugins/MovieScene/LevelSequenceEditor/Source/LevelSequenceEditor/Private/LevelSequenceEditorToolkit.cpp
@@ -120,11 +120,6 @@ FLevelSequenceEditorToolkit::~FLevelSequenceEditorToolkit()
 	if (FModuleManager::Get().IsModuleLoaded(TEXT("LevelEditor")))
 	{
 		FLevelEditorModule& LevelEditorModule = FModuleManager::LoadModuleChecked<FLevelEditorModule>("LevelEditor");
-		if (!IsWorldCentricAssetEditor())
-		{
-			TSharedPtr<FTabManager> LevelEditorTabManager = LevelEditorModule.GetLevelEditorTabManager();
-			LevelEditorTabManager->UnregisterTabSpawner(SequencerMainTabId);
-		}
 
 		// @todo remove when world-centric mode is added
 		LevelEditorModule.AttachSequencer(SNullWidget::NullWidget, nullptr);
@@ -229,18 +224,6 @@ void FLevelSequenceEditorToolkit::Initialize(const EToolkitMode::Type Mode, cons
 			LevelEditorTabManager->TryInvokeTab(FName("LevelEditorSceneOutliner"))->RequestCloseTab();
 			LevelEditorTabManager->TryInvokeTab(FName("LevelEditorSceneOutliner"));
 		}
-
-		// For World Centric Asset Editors this isn't called until way too late in the initialization flow
-		// (ie: when you actually start to edit an asset), so the tab will be unrecognized upon restore.
-		// Because of this, the Sequencer Tab Spawner is actually registered in SLevelEditor.cpp manually
-		// which is early enough that you can restore the tab after an editor restart.
-
-		WorkspaceMenuCategory = LevelEditorTabManager->AddLocalWorkspaceMenuCategory(LOCTEXT("WorkspaceMenu_SequencerAssetEditor", "Sequencer"));
-
-		LevelEditorTabManager->RegisterTabSpawner(SequencerMainTabId, FOnSpawnTab::CreateSP(this, &FLevelSequenceEditorToolkit::HandleTabManagerSpawnTab))
-			.SetDisplayName(LOCTEXT("SequencerMainTab", "Sequencer"))
-			.SetGroup(WorkspaceMenuCategory.ToSharedRef())
-			.SetIcon(FSlateIcon(Style->GetStyleSetName(), "LevelSequenceEditor.Tabs.Sequencer"));
 	}
 	
 	// Now Attach so this window will apear in the correct front first order
