Sequencer: Custom Curve Colors

Added Custom Color Support to Curve Editor. This is done via an Editor preference(CurveEditorSettings) that stores a color per UClass, property name(intent name of the channel in Curve Editor), and the Color. So would be something like StaticMeshActor, "Translate.X", (1.0f,0.0f,0.f90)

Added to Sections an UObject* GetImplicitObjectOwner(). Usually this is the Possessable/Spawnable of the owning track, but for Control Rig Tracks we use the
specified Control Rig instead.

So we use that implicit owner to get the UClass to associate with the channel.

Also added some Python/BP functions since editor pref's aren't scriptable.

diff --git a/Engine/Plugins/Animation/ControlRig/Source/ControlRig/Private/Sequencer/MovieSceneControlRigParameterSection.cpp b/Engine/Plugins/Animation/ControlRig/Source/ControlRig/Private/Sequencer/MovieSceneControlRigParameterSection.cpp
--- a/Engine/Plugins/Animation/ControlRig/Source/ControlRig/Private/Sequencer/MovieSceneControlRigParameterSection.cpp
+++ b/Engine/Plugins/Animation/ControlRig/Source/ControlRig/Private/Sequencer/MovieSceneControlRigParameterSection.cpp
@@ -2572,6 +2572,17 @@ void UMovieSceneControlRigParameterSection::ClearAllParameters()
 	EnumParameterNamesAndCurves.SetNum(0);
 	IntegerParameterNamesAndCurves.SetNum(0);
 }
+
+
+UObject* UMovieSceneControlRigParameterSection::GetImplicitObjectOwner()
+{
+	if (GetControlRig())
+	{
+		return GetControlRig();
+	}
+	return Super::GetImplicitObjectOwner();
+}
+
 FEnumParameterNameAndCurve::FEnumParameterNameAndCurve(FName InParameterName)
 {
 	ParameterName = InParameterName;
diff --git a/Engine/Plugins/Animation/ControlRig/Source/ControlRig/Public/Sequencer/MovieSceneControlRigParameterSection.h b/Engine/Plugins/Animation/ControlRig/Source/ControlRig/Public/Sequencer/MovieSceneControlRigParameterSection.h
--- a/Engine/Plugins/Animation/ControlRig/Source/ControlRig/Public/Sequencer/MovieSceneControlRigParameterSection.h
+++ b/Engine/Plugins/Animation/ControlRig/Source/ControlRig/Public/Sequencer/MovieSceneControlRigParameterSection.h
@@ -237,7 +237,10 @@ public:
 
 	UMovieSceneControlRigParameterSection();
 
+	//UMovieSceneSection virtuals
 	virtual void SetBlendType(EMovieSceneBlendType InBlendType) override;
+	virtual UObject* GetImplicitObjectOwner() override;
+
 #if WITH_EDITOR
 	//Function to save control rig key when recording.
 	void RecordControlRigKey(FFrameNumber FrameNumber, bool bSetDefault, bool bDoAutoKey);
diff --git a/Engine/Plugins/MovieScene/LevelSequenceEditor/Source/LevelSequenceEditor/LevelSequenceEditor.Build.cs b/Engine/Plugins/MovieScene/LevelSequenceEditor/Source/LevelSequenceEditor/LevelSequenceEditor.Build.cs
--- a/Engine/Plugins/MovieScene/LevelSequenceEditor/Source/LevelSequenceEditor/LevelSequenceEditor.Build.cs
+++ b/Engine/Plugins/MovieScene/LevelSequenceEditor/Source/LevelSequenceEditor/LevelSequenceEditor.Build.cs
@@ -21,6 +21,7 @@ public class LevelSequenceEditor : ModuleRules
                 "CinematicCamera",
 				"Core",
 				"CoreUObject",
+				"CurveEditor",
 				"EditorFramework",
                 "EditorStyle",
                 "Engine",
diff --git a/Engine/Plugins/MovieScene/LevelSequenceEditor/Source/LevelSequenceEditor/Private/LevelSequenceEditorBlueprintLibrary.cpp b/Engine/Plugins/MovieScene/LevelSequenceEditor/Source/LevelSequenceEditor/Private/LevelSequenceEditorBlueprintLibrary.cpp
--- a/Engine/Plugins/MovieScene/LevelSequenceEditor/Source/LevelSequenceEditor/Private/LevelSequenceEditorBlueprintLibrary.cpp
+++ b/Engine/Plugins/MovieScene/LevelSequenceEditor/Source/LevelSequenceEditor/Private/LevelSequenceEditorBlueprintLibrary.cpp
@@ -16,6 +16,9 @@
 #include "MovieSceneSection.h"
 #include "Channels/MovieSceneChannelProxy.h"
 
+//For custom colors on channels, stored in editor pref's
+#include "CurveEditorSettings.h"
+
 namespace
 {
 	static TWeakPtr<ISequencer> CurrentSequencer;
@@ -409,6 +412,82 @@ void ULevelSequenceEditorBlueprintLibrary::SetLockLevelSequence(bool bLock)
 	}
 }
 
+bool ULevelSequenceEditorBlueprintLibrary::HasCustomColorForChannel(UClass* Class, const FString& Identifier)
+{
+	const UCurveEditorSettings* Settings = GetDefault<UCurveEditorSettings>();
+	if (Settings)
+	{
+		TOptional<FLinearColor> OptColor = Settings->GetCustomColor(Class, Identifier);
+		return OptColor.IsSet();
+	}
+	return false;
+}
+
+FLinearColor ULevelSequenceEditorBlueprintLibrary::GetCustomColorForChannel(UClass* Class, const FString& Identifier)
+{
+	FLinearColor Color(FColor::White);
+	const UCurveEditorSettings* Settings = GetDefault<UCurveEditorSettings>();
+	if (Settings)
+	{
+		TOptional<FLinearColor> OptColor = Settings->GetCustomColor(Class, Identifier);
+		if (OptColor.IsSet())
+		{
+			return OptColor.GetValue();
+		}
+	}
+	return Color;
+}
+
+void ULevelSequenceEditorBlueprintLibrary::SetCustomColorForChannel(UClass* Class, const FString& Identifier, const FLinearColor& NewColor)
+{
+	UCurveEditorSettings* Settings = GetMutableDefault<UCurveEditorSettings>();
+	if (Settings)
+	{
+		Settings->SetCustomColor(Class, Identifier, NewColor);
+	}
+}
+
+void ULevelSequenceEditorBlueprintLibrary::SetCustomColorForChannels(UClass* Class, const TArray<FString>& Identifiers, const TArray<FLinearColor>& NewColors)
+{
+	if (Identifiers.Num() != NewColors.Num())
+	{
+		return;
+	}
+	UCurveEditorSettings* Settings = GetMutableDefault<UCurveEditorSettings>();
+	if (Settings)
+	{
+		for (int32 Index = 0; Index < Identifiers.Num(); ++Index)
+		{
+			const FString& Identifier = Identifiers[Index];
+			const FLinearColor& NewColor = NewColors[Index];
+			Settings->SetCustomColor(Class, Identifier, NewColor);
+		}
+	}
+}
+
+void ULevelSequenceEditorBlueprintLibrary::DeleteColorForChannels(UClass* Class, FString& Identifier)
+{
+	UCurveEditorSettings* Settings = GetMutableDefault<UCurveEditorSettings>();
+	if (Settings)
+	{
+		Settings->DeleteCustomColor(Class, Identifier);
+	}
+}
+
+void ULevelSequenceEditorBlueprintLibrary::SetRandomColorForChannels(UClass* Class, const TArray<FString>& Identifiers)
+{
+	UCurveEditorSettings* Settings = GetMutableDefault<UCurveEditorSettings>();
+	if (Settings)
+	{
+		for (int32 Index = 0; Index < Identifiers.Num(); ++Index)
+		{
+			const FString& Identifier = Identifiers[Index];
+			FLinearColor NewColor = UCurveEditorSettings::GetNextRandomColor();
+			Settings->SetCustomColor(Class, Identifier, NewColor);
+		}
+	}
+}
+
 bool ULevelSequenceEditorBlueprintLibrary::IsCameraCutLockedToViewport()
 {
 	if (CurrentSequencer.IsValid())
diff --git a/Engine/Plugins/MovieScene/LevelSequenceEditor/Source/LevelSequenceEditor/Public/LevelSequenceEditorBlueprintLibrary.h b/Engine/Plugins/MovieScene/LevelSequenceEditor/Source/LevelSequenceEditor/Public/LevelSequenceEditorBlueprintLibrary.h
--- a/Engine/Plugins/MovieScene/LevelSequenceEditor/Source/LevelSequenceEditor/Public/LevelSequenceEditorBlueprintLibrary.h
+++ b/Engine/Plugins/MovieScene/LevelSequenceEditor/Source/LevelSequenceEditor/Public/LevelSequenceEditorBlueprintLibrary.h
@@ -207,6 +207,31 @@ public:
 	UFUNCTION(BlueprintCallable, Category = "Level Sequence Editor")
 	static void SetLockCameraCutToViewport(bool bLock);
 
+public:
+	/** Get if a custom color for specified channel idendified by it's class and identifier exists */
+	UFUNCTION(BlueprintCallable, Category = "Level Sequence Editor")
+	static bool HasCustomColorForChannel(UClass* Class, const FString& Identifier);
+	
+	/** Get custom color for specified channel idendified by it's class and identifier,if none exists will return white*/
+	UFUNCTION(BlueprintCallable, Category = "Level Sequence Editor")
+	static FLinearColor GetCustomColorForChannel(UClass* Class, const FString& Identifier);
+	
+	/** Set Custom Color for specified channel idendified by it's class and identifier. This will be stored in editor user preferences.*/
+	UFUNCTION(BlueprintCallable, Category = "Level Sequence Editor")
+	static void SetCustomColorForChannel(UClass* Class, const FString& Identifier, const FLinearColor& NewColor);
+	
+	/** Set Custom Color for specified channels idendified by it's class and identifiers. This will be stored in editor user preferences.*/
+	UFUNCTION(BlueprintCallable, Category = "Level Sequence Editor")
+	static void SetCustomColorForChannels(UClass* Class, const TArray<FString>& Identifiers, const TArray<FLinearColor>& NewColors);
+	
+	/** Set Random Colors for specified channels idendified by it's class and identifiers. This will be stored in editor user preferences.*/
+	UFUNCTION(BlueprintCallable, Category = "Level Sequence Editor")
+	static void SetRandomColorForChannels(UClass* Class, const TArray<FString>& Identifiers);
+	
+	/** Delete for specified channel idendified by it's class and identifier.*/
+	UFUNCTION(BlueprintCallable, Category = "Level Sequence Editor")
+	static void DeleteColorForChannels(UClass* Class, FString& Identifier);
+
 public:
 
 	/*
diff --git a/Engine/Source/Editor/CurveEditor/CurveEditor.Build.cs b/Engine/Source/Editor/CurveEditor/CurveEditor.Build.cs
--- a/Engine/Source/Editor/CurveEditor/CurveEditor.Build.cs
+++ b/Engine/Source/Editor/CurveEditor/CurveEditor.Build.cs
@@ -9,12 +9,13 @@ public class CurveEditor : ModuleRules
 		PrivateDependencyModuleNames.AddRange(
 			new string[] {
                 "ApplicationCore",
+                "AppFramework",
                 "Core",
 				"CoreUObject",
 				"EditorFramework",
 				"Engine",
 				"EditorStyle",
-				"InputCore",
+                "InputCore",
 				"Slate",
 				"SlateCore",
 				"TimeManagement",
diff --git a/Engine/Source/Editor/CurveEditor/Private/CurveEditor.cpp b/Engine/Source/Editor/CurveEditor/Private/CurveEditor.cpp
--- a/Engine/Source/Editor/CurveEditor/Private/CurveEditor.cpp
+++ b/Engine/Source/Editor/CurveEditor/Private/CurveEditor.cpp
@@ -28,6 +28,7 @@
 #include "Factories.h"
 #include "HAL/PlatformApplicationMisc.h"
 #include "SCurveEditor.h" // for access to LogCurveEditor
+#include "Widgets/Colors/SColorPicker.h"
 
 #define LOCTEXT_NAMESPACE "CurveEditor"
 
@@ -373,6 +374,11 @@ void FCurveEditor::BindCommands()
 		CommandList->MapAction(FCurveEditorCommands::Get().StraightenTangents, FExecuteAction::CreateSP(this, &FCurveEditor::StraightenSelection), FCanExecuteAction::CreateSP(this, &FCurveEditor::CanFlattenOrStraightenSelection) );
 	}
 
+	// Curve Colors
+	{
+		CommandList->MapAction(FCurveEditorCommands::Get().SetRandomCurveColorsForSelected, FExecuteAction::CreateSP(this, &FCurveEditor::SetRandomCurveColorsForSelected), FCanExecuteAction());
+		CommandList->MapAction(FCurveEditorCommands::Get().SetCurveColorsForSelected, FExecuteAction::CreateSP(this, &FCurveEditor::SetCurveColorsForSelected), FCanExecuteAction());
+	}
 
 	// Tangent Visibility
 	{
@@ -405,7 +411,7 @@ void FCurveEditor::BindCommands()
 		FExecuteAction::CreateSP(this, &FCurveEditor::MakeToolActive, FCurveEditorToolID::Unset()),
 		FCanExecuteAction(),
 		FIsActionChecked::CreateLambda( [this]{ return ActiveTool.IsSet() == false; } ) );
-		
+
 	// Bind commands for Editor Extensions
 	for (TSharedRef<ICurveEditorExtension> Extension : EditorExtensions)
 	{
@@ -1546,6 +1552,54 @@ bool FCurveEditor::CanFlattenOrStraightenSelection() const
 	return Selection.Count() > 0;
 }
 
+void FCurveEditor::SetRandomCurveColorsForSelected()
+{
+	for (TPair<FCurveModelID, TUniquePtr<FCurveModel>>& CurvePair : CurveData)
+	{
+		if (FCurveModel* Curve = CurvePair.Value.Get())
+		{
+			UObject* Object = nullptr;
+			FString Name;
+			Curve->GetCurveColorObjectAndName(&Object, Name);
+			if (Object)
+			{
+				FLinearColor Color = UCurveEditorSettings::GetNextRandomColor();
+				Settings->SetCustomColor(Object->GetClass(), Name, Color);
+				Curve->SetColor(Color);
+			}
+		}
+	}
+}
+
+void FCurveEditor::SetCurveColorsForSelected()
+{
+	if (CurveData.Num() > 0)
+	{
+		TMap<FCurveModelID, TUniquePtr<FCurveModel>>::TIterator It = CurveData.CreateIterator();
+		FColorPickerArgs PickerArgs;
+		PickerArgs.bUseAlpha = false;
+		PickerArgs.InitialColorOverride = It->Value->GetColor();
+		PickerArgs.OnColorCommitted.BindLambda([this](FLinearColor NewColor) {
+			for (TPair<FCurveModelID, TUniquePtr<FCurveModel>>& CurvePair : CurveData)
+			{
+				if (FCurveModel* Curve = CurvePair.Value.Get())
+				{
+					UObject* Object = nullptr;
+					FString Name;
+					Curve->GetCurveColorObjectAndName(&Object, Name);
+					if (Object)
+					{
+						Settings->SetCustomColor(Object->GetClass(), Name, NewColor);
+						Curve->SetColor(NewColor);
+					}
+				}
+			}
+		});
+		
+		OpenColorPicker(PickerArgs);
+	}
+}
+
 bool FCurveEditor::IsToolActive(const FCurveEditorToolID InToolID) const
 {
 	if (ActiveTool.IsSet())
diff --git a/Engine/Source/Editor/CurveEditor/Private/CurveEditorCommands.cpp b/Engine/Source/Editor/CurveEditor/Private/CurveEditorCommands.cpp
--- a/Engine/Source/Editor/CurveEditor/Private/CurveEditorCommands.cpp
+++ b/Engine/Source/Editor/CurveEditor/Private/CurveEditorCommands.cpp
@@ -57,6 +57,10 @@ void FCurveEditorCommands::RegisterCommands()
 
 	UI_COMMAND(AddKeyToAllCurves, "Add Key", "Add a new key to all curves at the current time.", EUserInterfaceActionType::Button, FInputChord(EKeys::Enter) );
 
+	// Curve Editor Colors
+	UI_COMMAND(SetRandomCurveColorsForSelected, "Set Random Curve Colors", "Set random colors on the selected curves. Note they are stored in the Level Sequence Actor Editor Preferences.", EUserInterfaceActionType::Button, FInputChord());
+	UI_COMMAND(SetCurveColorsForSelected, "Set Curve Color For Selected", "Set the chosen color on the selected curves. Note they are stored in the Level Sequence Actor Editor Preferences.", EUserInterfaceActionType::Button, FInputChord());
+
 	// Graph Viewing Modes
 	UI_COMMAND(SetViewModeAbsolute, "Absolute View Mode", "Absolute view displays all curves overlapping with the Y axis proportionally scaled.", EUserInterfaceActionType::ToggleButton, FInputChord());
 	UI_COMMAND(SetViewModeStacked, "Stacked View Mode", "Stacked view displays each curve in its own graph with the Y axis normalized [-1, 1].", EUserInterfaceActionType::ToggleButton, FInputChord());
diff --git a/Engine/Source/Editor/CurveEditor/Private/CurveEditorSettings.cpp b/Engine/Source/Editor/CurveEditor/Private/CurveEditorSettings.cpp
--- a/Engine/Source/Editor/CurveEditor/Private/CurveEditorSettings.cpp
+++ b/Engine/Source/Editor/CurveEditor/Private/CurveEditorSettings.cpp
@@ -66,3 +66,77 @@ void UCurveEditorSettings::SetZoomPosition(ECurveEditorZoomPosition InZoomPositi
 	}
 }
 
+TOptional<FLinearColor> UCurveEditorSettings::GetCustomColor(UClass* InClass, const FString& InPropertyName) const
+{
+	TOptional<FLinearColor> Color;
+	for (const FCustomColorForChannel& CustomColor : CustomColors)
+	{
+		UClass* Class = CustomColor.Object.LoadSynchronous();
+		if (Class == InClass && CustomColor.PropertyName == InPropertyName)
+		{
+			Color = CustomColor.Color;
+			break;
+		}
+	}
+	return Color;
+}
+void UCurveEditorSettings::SetCustomColor(UClass* InClass, const FString& InPropertyName, FLinearColor InColor)
+{
+	TOptional<FLinearColor> Color;
+	for (FCustomColorForChannel& CustomColor : CustomColors)
+	{
+		UClass* Class = CustomColor.Object.LoadSynchronous();
+		if (Class == InClass && CustomColor.PropertyName == InPropertyName)
+		{
+			CustomColor.Color = InColor;
+			SaveConfig(); 
+			return;
+		}
+	}
+	FCustomColorForChannel NewColor;
+	NewColor.Object = InClass;
+	NewColor.PropertyName = InPropertyName;
+	NewColor.Color = InColor;
+	CustomColors.Add(NewColor);
+	SaveConfig();
+}
+
+void UCurveEditorSettings::DeleteCustomColor(UClass* InClass, const FString& InPropertyName)
+{
+	TOptional<FLinearColor> Color;
+	for (int32 Index = 0;Index < CustomColors.Num(); ++Index)
+	{
+		FCustomColorForChannel& CustomColor = CustomColors[Index];
+		UClass* Class = CustomColor.Object.LoadSynchronous();
+		if (Class == InClass && CustomColor.PropertyName == InPropertyName)
+		{
+			CustomColors.RemoveAt(Index);
+			SaveConfig();
+			return;
+		}
+	}
+}
+
+FLinearColor UCurveEditorSettings::GetNextRandomColor()
+{
+	static TArray<FLinearColor> IndexedColor;
+	static int32 NextIndex = 0;
+	if (IndexedColor.Num() == 0)
+	{
+		IndexedColor.Add(FLinearColor(FColor::Magenta));
+		IndexedColor.Add(FLinearColor(FColor::Cyan));
+		IndexedColor.Add(FLinearColor(FColor::Turquoise));
+		IndexedColor.Add(FLinearColor(FColor::Orange));
+		IndexedColor.Add(FLinearColor(FColor::Yellow));
+		IndexedColor.Add(FLinearColor(FColor::Purple));
+		IndexedColor.Add(FLinearColor(FColor::Silver));
+		IndexedColor.Add(FLinearColor(FColor::Emerald));
+		IndexedColor.Add(FLinearColor(FColor::White));
+		IndexedColor.Add(FLinearColor(FColor::Red));
+		IndexedColor.Add(FLinearColor(FColor::Green));
+		IndexedColor.Add(FLinearColor(FColor::Blue));
+	}
+	FLinearColor Color = IndexedColor[NextIndex];
+	NextIndex = (++NextIndex) % IndexedColor.Num();
+	return Color;
+}
\ No newline at end of file
diff --git a/Engine/Source/Editor/CurveEditor/Private/SCurveEditorPanel.cpp b/Engine/Source/Editor/CurveEditor/Private/SCurveEditorPanel.cpp
--- a/Engine/Source/Editor/CurveEditor/Private/SCurveEditorPanel.cpp
+++ b/Engine/Source/Editor/CurveEditor/Private/SCurveEditorPanel.cpp
@@ -899,6 +899,13 @@ TSharedRef<SWidget> SCurveEditorPanel::MakeCurveEditorCurveViewOptionsMenu()
 	}
 	MenuBuilder.EndSection();
 
+	MenuBuilder.BeginSection("CurveColors", LOCTEXT("CurveColorsHeader", "Curve Colors"));
+	{
+		MenuBuilder.AddMenuEntry(FCurveEditorCommands::Get().SetRandomCurveColorsForSelected);
+		MenuBuilder.AddMenuEntry(FCurveEditorCommands::Get().SetCurveColorsForSelected);
+	}
+	MenuBuilder.EndSection();
+
 	return MenuBuilder.MakeWidget();
 }
 
diff --git a/Engine/Source/Editor/CurveEditor/Public/CurveEditor.h b/Engine/Source/Editor/CurveEditor/Public/CurveEditor.h
--- a/Engine/Source/Editor/CurveEditor/Public/CurveEditor.h
+++ b/Engine/Source/Editor/CurveEditor/Public/CurveEditor.h
@@ -439,6 +439,16 @@ public:
 	 */
 	void StraightenSelection();
 
+	/**
+	 * Set random curve colors
+	 */
+	void SetRandomCurveColorsForSelected();
+
+	/**
+	 * Pick a curve color and set on selected
+	 */	
+	void SetCurveColorsForSelected();
+
 	/**
 	* Do we currently have keys to flatten or straighten?
 	*/
diff --git a/Engine/Source/Editor/CurveEditor/Public/CurveEditorCommands.h b/Engine/Source/Editor/CurveEditor/Public/CurveEditorCommands.h
--- a/Engine/Source/Editor/CurveEditor/Public/CurveEditorCommands.h
+++ b/Engine/Source/Editor/CurveEditor/Public/CurveEditorCommands.h
@@ -66,6 +66,9 @@ public:
 	TSharedPtr<FUICommandInfo> ToggleAutoFrameCurveEditor;
 	TSharedPtr<FUICommandInfo> ToggleShowCurveEditorCurveToolTips;
 
+	TSharedPtr<FUICommandInfo> SetRandomCurveColorsForSelected;
+	TSharedPtr<FUICommandInfo> SetCurveColorsForSelected;
+
 	TSharedPtr<FUICommandInfo> AddKeyHovered;
 	TSharedPtr<FUICommandInfo> PasteKeysHovered;
 
diff --git a/Engine/Source/Editor/CurveEditor/Public/CurveEditorSettings.h b/Engine/Source/Editor/CurveEditor/Public/CurveEditorSettings.h
--- a/Engine/Source/Editor/CurveEditor/Public/CurveEditorSettings.h
+++ b/Engine/Source/Editor/CurveEditor/Public/CurveEditorSettings.h
@@ -30,6 +30,24 @@ enum class ECurveEditorZoomPosition : uint8
 	MousePosition,
 };
 
+/** Custom Color Object*/
+USTRUCT()
+struct FCustomColorForChannel
+{
+	GENERATED_BODY()
+
+	FCustomColorForChannel() : Object(nullptr), Color(1.0f, 1.0f, 1.0f, 1.0f) {};
+
+	UPROPERTY(EditAnywhere, Category = "Custom Colors")
+	TSoftClassPtr<UObject>	Object;
+	UPROPERTY(EditAnywhere, Category = "Custom Colors")
+	FString PropertyName;
+	UPROPERTY(EditAnywhere, Category = "Custom Colors")
+	FLinearColor Color;
+
+
+};
+
 /** Serializable options for curve editor. */
 UCLASS(config=EditorPerProjectUserSettings)
 class CURVEEDITOR_API UCurveEditorSettings : public UObject
@@ -59,6 +77,16 @@ public:
 	/** Set zoom in/out position (mouse position or current time). */
 	void SetZoomPosition(ECurveEditorZoomPosition InZoomPosition);
 
+	/** Get custom color for object and property if it exists, if it doesn't the optional won't be set */
+	TOptional<FLinearColor> GetCustomColor(UClass* InClass, const FString& InPropertyName) const;
+	/** Set Custom Color for the specified parameters. */
+	void SetCustomColor(UClass* InClass, const FString& InPropertyName, FLinearColor InColor);
+	/** Delete Custom Color for the specified parameters. */
+	void DeleteCustomColor(UClass* InClass, const FString& InPropertyName);
+
+	/** Helper function to get next random linear color*/
+	static FLinearColor GetNextRandomColor();
+
 protected:
 	UPROPERTY( config, EditAnywhere, Category="Curve Editor" )
 	bool bAutoFrameCurveEditor;
@@ -71,4 +99,7 @@ protected:
 
 	UPROPERTY( config, EditAnywhere, Category="Curve Editor")
 	ECurveEditorZoomPosition ZoomPosition;
+
+	UPROPERTY(config, EditAnywhere, Category="Curve Editor")
+	TArray<FCustomColorForChannel> CustomColors;
 };
diff --git a/Engine/Source/Editor/CurveEditor/Public/CurveModel.h b/Engine/Source/Editor/CurveEditor/Public/CurveModel.h
--- a/Engine/Source/Editor/CurveEditor/Public/CurveModel.h
+++ b/Engine/Source/Editor/CurveEditor/Public/CurveModel.h
@@ -240,6 +240,16 @@ public:
 	{
 		return nullptr;
 	}
+
+	/**
+	* Get the Object and the name to be used to store the curve model color (see UCurveEditorSettings). By default
+	* this is the owning object and the intent name, but it can be overriden, for example for Sequencer it may be the bound object
+	*/
+	virtual void GetCurveColorObjectAndName(UObject** OutObject, FString& OutName) const
+	{
+		*OutObject = GetOwningObject();
+		OutName = GetIntentionName();
+	}
 	/**
 	 * Helper function for assigning a the same attributes to a number of keys
 	 */
diff --git a/Engine/Source/Editor/MovieSceneTools/Private/Channels/ChannelCurveModel.cpp b/Engine/Source/Editor/MovieSceneTools/Private/Channels/ChannelCurveModel.cpp
--- a/Engine/Source/Editor/MovieSceneTools/Private/Channels/ChannelCurveModel.cpp
+++ b/Engine/Source/Editor/MovieSceneTools/Private/Channels/ChannelCurveModel.cpp
@@ -425,6 +425,19 @@ void FChannelCurveModel<ChannelType, ChannelValue, KeyType>::FixupCurve()
 		OnDestroyHandle = NewChannelProxy->OnDestroy.AddRaw(this, &FChannelCurveModel<ChannelType, ChannelValue, KeyType>::FixupCurve);
 	}
 }
+template <class ChannelType, class ChannelValue, class KeyType>
+void FChannelCurveModel<ChannelType, ChannelValue, KeyType>::GetCurveColorObjectAndName(UObject** OutObject, FString& OutName) const
+{
+	if (UMovieSceneSection* Section = WeakSection.Get())
+	{
+
+		*OutObject = Section->GetImplicitObjectOwner();
+		OutName = GetIntentionName();
+	}
+	// Just call base if it doesn't work
+	FCurveModel::GetCurveColorObjectAndName(OutObject, OutName);
+}
+
 
 // Explicit template instantiation
 template class FChannelCurveModel<FMovieSceneDoubleChannel, FMovieSceneDoubleValue, double>;
diff --git a/Engine/Source/Editor/MovieSceneTools/Public/Channels/ChannelCurveModel.h b/Engine/Source/Editor/MovieSceneTools/Public/Channels/ChannelCurveModel.h
--- a/Engine/Source/Editor/MovieSceneTools/Public/Channels/ChannelCurveModel.h
+++ b/Engine/Source/Editor/MovieSceneTools/Public/Channels/ChannelCurveModel.h
@@ -46,6 +46,7 @@ public:
 	{
 		return WeakSection.Get();
 	}
+	virtual void GetCurveColorObjectAndName(UObject** OutObject, FString& OutName) const override;
 
 protected:
 
diff --git a/Engine/Source/Editor/Sequencer/Private/IKeyArea.cpp b/Engine/Source/Editor/Sequencer/Private/IKeyArea.cpp
--- a/Engine/Source/Editor/Sequencer/Private/IKeyArea.cpp
+++ b/Engine/Source/Editor/Sequencer/Private/IKeyArea.cpp
@@ -13,6 +13,7 @@
 #include "ISequencer.h"
 #include "ISequencerSection.h"
 #include "MovieSceneSequence.h"
+#include "CurveEditorSettings.h"
 
 IKeyArea::IKeyArea(TWeakPtr<ISequencerSection> InSection, FMovieSceneChannelHandle InChannel)
 	: TreeSerialNumber(0)
@@ -292,11 +293,24 @@ TUniquePtr<FCurveModel> IKeyArea::CreateCurveEditorModel(TSharedRef<ISequencer>
 			CurveModel->SetShortDisplayName(DisplayText);
 			CurveModel->SetLongDisplayName(LongDisplayName);
 			CurveModel->SetIntentionName(IntentName);
-
 			if (Color.IsSet())
 			{
 				CurveModel->SetColor(Color.GetValue());
 			}
+
+			//Use editor preference color if it's been set, this function is const so we can't set just Color optional
+			const UCurveEditorSettings* Settings = GetDefault<UCurveEditorSettings>();
+			UObject* Object = nullptr;
+			FString Name;
+			CurveModel->GetCurveColorObjectAndName(&Object, Name);
+			if (Object)
+			{
+				TOptional<FLinearColor> SettingColor = Settings->GetCustomColor(Object->GetClass(), Name);
+				if (SettingColor.IsSet())
+				{
+					CurveModel->SetColor(SettingColor.GetValue());
+				}
+			}
 		}
 		return CurveModel;
 	}
diff --git a/Engine/Source/Runtime/MovieScene/Private/MovieSceneSection.cpp b/Engine/Source/Runtime/MovieScene/Private/MovieSceneSection.cpp
--- a/Engine/Source/Runtime/MovieScene/Private/MovieSceneSection.cpp
+++ b/Engine/Source/Runtime/MovieScene/Private/MovieSceneSection.cpp
@@ -571,6 +571,32 @@ UMovieSceneSection* UMovieSceneSection::SplitSection(FQualifiedFrameTime SplitTi
 	return nullptr;
 }
 
+UObject* UMovieSceneSection::GetImplicitObjectOwner()
+{
+	if (UMovieSceneTrack* Track = GetTypedOuter<UMovieSceneTrack>())
+	{
+		if (UMovieScene* MovieScene = Track->GetTypedOuter<UMovieScene>())
+		{
+			FGuid Guid;
+			if (MovieScene->FindTrackBinding(*Track, Guid))
+			{
+				if (FMovieSceneSpawnable* MovieSceneSpanwable = MovieScene->FindSpawnable(Guid))
+				{
+					return MovieSceneSpanwable->GetObjectTemplate();
+				}
+				else if (FMovieScenePossessable* MovieScenePossessable = MovieScene->FindPossessable(Guid))
+				{
+					if (MovieScenePossessable->GetPossessedObjectClass() && MovieScenePossessable->GetPossessedObjectClass()->GetDefaultObject())
+					{
+						return MovieScenePossessable->GetPossessedObjectClass()->GetDefaultObject();
+					}
+				}
+			}
+		}
+	}
+	return nullptr;
+}
+
 
 void UMovieSceneSection::TrimSection(FQualifiedFrameTime TrimTime, bool bTrimLeft, bool bDeleteKeys)
 {
diff --git a/Engine/Source/Runtime/MovieScene/Public/MovieSceneSection.h b/Engine/Source/Runtime/MovieScene/Public/MovieSceneSection.h
--- a/Engine/Source/Runtime/MovieScene/Public/MovieSceneSection.h
+++ b/Engine/Source/Runtime/MovieScene/Public/MovieSceneSection.h
@@ -561,6 +561,15 @@ public:
 	*/
 	MOVIESCENE_API virtual float GetTotalWeightValue(FFrameTime InTime) const { return EvaluateEasing(InTime); }
 
+
+	/**
+	*  Get the implicit owner of this section, usually this will be the section's outer possessable or spawnable,
+	*  but some sections, like Control Rig, this will be the Control Rig object instead.
+	*
+	**/
+	MOVIESCENE_API virtual UObject* GetImplicitObjectOwner();
+
+
 #if WITH_EDITOR
 	MOVIESCENE_API virtual void PostEditChangeProperty(FPropertyChangedEvent& PropertyChangedEvent) override;
 #endif
