Sequencer: Fix include

diff --git a/Engine/Plugins/MovieScene/LevelSequenceEditor/Source/LevelSequenceEditor/Private/LevelSequenceEditorSubsystem.cpp b/Engine/Plugins/MovieScene/LevelSequenceEditor/Source/LevelSequenceEditor/Private/LevelSequenceEditorSubsystem.cpp
--- a/Engine/Plugins/MovieScene/LevelSequenceEditor/Source/LevelSequenceEditor/Private/LevelSequenceEditorSubsystem.cpp
+++ b/Engine/Plugins/MovieScene/LevelSequenceEditor/Source/LevelSequenceEditor/Private/LevelSequenceEditorSubsystem.cpp
@@ -10,6 +10,7 @@
 #include "LevelSequenceEditorCommands.h"
 #include "ActorTreeItem.h"
 #include "LevelEditor.h"
+#include "ScopedTransaction.h"
 #include "Compilation/MovieSceneCompiledDataManager.h"
 
 
