Sequencer: Expose BakeTransform to python/blueprints by moving it to LevelSequenceEditorSubsystem

diff --git a/Engine/Plugins/MovieScene/LevelSequenceEditor/Source/LevelSequenceEditor/LevelSequenceEditor.Build.cs b/Engine/Plugins/MovieScene/LevelSequenceEditor/Source/LevelSequenceEditor/LevelSequenceEditor.Build.cs
--- a/Engine/Plugins/MovieScene/LevelSequenceEditor/Source/LevelSequenceEditor/LevelSequenceEditor.Build.cs
+++ b/Engine/Plugins/MovieScene/LevelSequenceEditor/Source/LevelSequenceEditor/LevelSequenceEditor.Build.cs
@@ -12,6 +12,12 @@ public class LevelSequenceEditor : ModuleRules
 				"PlacementMode",
 			}
         );
+
+		PublicDependencyModuleNames.AddRange(
+			new string[] {
+				"SequencerScripting",
+			}
+		);
         
         PrivateDependencyModuleNames.AddRange(
 			new string[] {
diff --git a/Engine/Plugins/MovieScene/LevelSequenceEditor/Source/LevelSequenceEditor/Private/LevelSequenceEditorCommands.cpp b/Engine/Plugins/MovieScene/LevelSequenceEditor/Source/LevelSequenceEditor/Private/LevelSequenceEditorCommands.cpp
--- a/Engine/Plugins/MovieScene/LevelSequenceEditor/Source/LevelSequenceEditor/Private/LevelSequenceEditorCommands.cpp
+++ b/Engine/Plugins/MovieScene/LevelSequenceEditor/Source/LevelSequenceEditor/Private/LevelSequenceEditorCommands.cpp
@@ -15,8 +15,8 @@ void FLevelSequenceEditorCommands::RegisterCommands()
 	UI_COMMAND(CreateNewMasterSequenceInLevel, "Add Master Sequence", "Create a new master sequence asset, and place an instance of it in this level", EUserInterfaceActionType::Button, FInputChord());
 	UI_COMMAND(ToggleCinematicViewportCommand, "Cinematic Viewport", "A viewport layout tailored to cinematic preview", EUserInterfaceActionType::RadioButton, FInputChord());
 
-	UI_COMMAND(FixActorReferences, "Fix Actor References", "Try to automatically fix up broken actor bindings.", EUserInterfaceActionType::Button, FInputChord());
-
+	UI_COMMAND(BakeTransform, "Bake Transform", "Bake transform in world space, removing any existing transform and attach tracks", EUserInterfaceActionType::Button, FInputChord());
+	UI_COMMAND(FixActorReferences, "Fix Actor References", "Try to automatically fix up broken actor bindings", EUserInterfaceActionType::Button, FInputChord());
 }
 
 #undef LOCTEXT_NAMESPACE
diff --git a/Engine/Plugins/MovieScene/LevelSequenceEditor/Source/LevelSequenceEditor/Private/LevelSequenceEditorCommands.h b/Engine/Plugins/MovieScene/LevelSequenceEditor/Source/LevelSequenceEditor/Private/LevelSequenceEditorCommands.h
--- a/Engine/Plugins/MovieScene/LevelSequenceEditor/Source/LevelSequenceEditor/Private/LevelSequenceEditorCommands.h
+++ b/Engine/Plugins/MovieScene/LevelSequenceEditor/Source/LevelSequenceEditor/Private/LevelSequenceEditorCommands.h
@@ -20,7 +20,6 @@ public:
 	TSharedPtr<FUICommandInfo> CreateNewMasterSequenceInLevel;
 	TSharedPtr<FUICommandInfo> ToggleCinematicViewportCommand;
 	
-	/** Attempts to fix broken actor references. */
-	TSharedPtr< FUICommandInfo > FixActorReferences;
-
+	TSharedPtr<FUICommandInfo> BakeTransform;
+	TSharedPtr<FUICommandInfo> FixActorReferences;
 };
diff --git a/Engine/Plugins/MovieScene/LevelSequenceEditor/Source/LevelSequenceEditor/Private/LevelSequenceEditorSubsystem.cpp b/Engine/Plugins/MovieScene/LevelSequenceEditor/Source/LevelSequenceEditor/Private/LevelSequenceEditorSubsystem.cpp
--- a/Engine/Plugins/MovieScene/LevelSequenceEditor/Source/LevelSequenceEditor/Private/LevelSequenceEditorSubsystem.cpp
+++ b/Engine/Plugins/MovieScene/LevelSequenceEditor/Source/LevelSequenceEditor/Private/LevelSequenceEditorSubsystem.cpp
@@ -1,22 +1,39 @@
 // Copyright Epic Games, Inc. All Rights Reserved.
 
 #include "LevelSequenceEditorSubsystem.h"
-#include "EngineUtils.h"
+
+#include "Compilation/MovieSceneCompiledDataManager.h"
 #include "ISequencer.h"
 #include "ISequencerModule.h"
-#include "SequencerUtilities.h"
+#include "LevelSequenceEditorCommands.h"
 #include "MovieScene.h"
+#include "MovieSceneCommonHelpers.h"
 #include "MovieSceneSequence.h"
-#include "LevelSequenceEditorCommands.h"
+#include "MovieSceneTimeHelpers.h"
+#include "SequencerSettings.h"
+#include "SequencerUtilities.h"
+#include "SequencerBindingProxy.h"
+#include "SequenceTimeUnit.h"
+#include "Sections/MovieScene3DConstraintSection.h"
+#include "Sections/MovieScene3DTransformSection.h"
+#include "Tracks/MovieScene3DTransformTrack.h"
+#include "Tracks/MovieScene3DConstraintTrack.h"
+#include "Tracks/MovieSceneCameraAnimTrack.h"
+#include "Tracks/MovieSceneCameraShakeTrack.h"
+
 #include "ActorTreeItem.h"
+#include "Camera/CameraComponent.h"
+#include "Components/SkeletalMeshComponent.h"
+#include "EngineUtils.h"
+#include "Framework/MultiBox/MultiBoxBuilder.h"
 #include "LevelEditor.h"
-#include "ScopedTransaction.h"
-#include "Compilation/MovieSceneCompiledDataManager.h"
 #include "Modules/ModuleManager.h"
-#include "Framework/MultiBox/MultiBoxBuilder.h"
+#include "ScopedTransaction.h"
 
 DEFINE_LOG_CATEGORY(LogLevelSequenceEditor);
 
+#define LOCTEXT_NAMESPACE "LevelSequenceEditor"
+
 void ULevelSequenceEditorSubsystem::Initialize(FSubsystemCollectionBase& Collection)
 {
 	UE_LOG(LogLevelSequenceEditor, Log, TEXT("LevelSequenceEditor subsystem initialized."));
@@ -26,11 +43,22 @@ void ULevelSequenceEditorSubsystem::Initialize(FSubsystemCollectionBase& Collect
 
 	/* Commands for this subsystem */
 	CommandList = MakeShareable(new FUICommandList);
+
+	CommandList->MapAction(FLevelSequenceEditorCommands::Get().BakeTransform,
+		FExecuteAction::CreateUObject(this, &ULevelSequenceEditorSubsystem::BakeTransformInternal)
+	);
+
 	CommandList->MapAction(FLevelSequenceEditorCommands::Get().FixActorReferences,
 		FExecuteAction::CreateUObject(this, &ULevelSequenceEditorSubsystem::FixActorReferences)
 	);
 
-	/* Createand register "Fix Actor References" extension */
+	BakeTransformMenuExtender = MakeShareable(new FExtender);
+	BakeTransformMenuExtender->AddMenuExtension("Transform", EExtensionHook::After, CommandList, FMenuExtensionDelegate::CreateStatic([](FMenuBuilder& MenuBuilder) {
+		MenuBuilder.AddMenuEntry(FLevelSequenceEditorCommands::Get().BakeTransform);
+		}));
+
+	SequencerModule.GetActionsMenuExtensibilityManager()->AddExtender(BakeTransformMenuExtender);
+
 	FixActorReferencesMenuExtender = MakeShareable(new FExtender);
 	FixActorReferencesMenuExtender->AddMenuExtension("Bindings", EExtensionHook::First, CommandList, FMenuExtensionDelegate::CreateStatic([](FMenuBuilder& MenuBuilder) {
 		MenuBuilder.AddMenuEntry(FLevelSequenceEditorCommands::Get().FixActorReferences);
@@ -71,6 +99,377 @@ TSharedPtr<ISequencer> ULevelSequenceEditorSubsystem::GetActiveSequencer()
 	return nullptr;
 }
 
+void ULevelSequenceEditorSubsystem::BakeTransformInternal()
+{
+	TSharedPtr<ISequencer> Sequencer = GetActiveSequencer();
+	if (Sequencer == nullptr)
+	{
+		return;
+	}
+
+	UMovieScene* FocusedMovieScene = Sequencer->GetFocusedMovieSceneSequence()->GetMovieScene();
+	if (!FocusedMovieScene)
+	{
+		return;
+	}
+
+	FFrameRate TickResolution = FocusedMovieScene->GetTickResolution();
+	FFrameRate DisplayRate = FocusedMovieScene->GetDisplayRate();
+
+	FFrameNumber InFrame = UE::MovieScene::DiscreteInclusiveLower(FocusedMovieScene->GetPlaybackRange());
+	FFrameNumber OutFrame = UE::MovieScene::DiscreteExclusiveUpper(FocusedMovieScene->GetPlaybackRange());
+
+	TArray<FGuid> ObjectBindings;
+	Sequencer->GetSelectedObjects(ObjectBindings);
+
+	TArray<FSequencerBindingProxy> BindingProxies;
+	for (FGuid Guid : ObjectBindings)
+	{
+		BindingProxies.Add(FSequencerBindingProxy(Guid, Sequencer->GetFocusedMovieSceneSequence()));
+	}
+
+	FMovieSceneScriptingParams Params;
+	Params.TimeUnit = ESequenceTimeUnit::TickResolution;
+	BakeTransform(BindingProxies, FFrameTime(InFrame), FFrameTime(OutFrame), ConvertFrameTime(1, DisplayRate, TickResolution), Params);
+}
+
+void ULevelSequenceEditorSubsystem::BakeTransform(const TArray<FSequencerBindingProxy>& ObjectBindings, const FFrameTime& BakeInTime, const FFrameTime& BakeOutTime, const FFrameTime& BakeInterval, const FMovieSceneScriptingParams& Params)
+{
+	TSharedPtr<ISequencer> Sequencer = GetActiveSequencer();
+	if (Sequencer == nullptr)
+	{
+		return;
+	}
+
+	UMovieScene* FocusedMovieScene = Sequencer->GetFocusedMovieSceneSequence()->GetMovieScene();
+	if (!FocusedMovieScene)
+	{
+		return;
+	}
+
+	if (FocusedMovieScene->IsReadOnly())
+	{
+		FSequencerUtilities::ShowReadOnlyError();
+		return;
+	}
+
+	if (ObjectBindings.Num() == 0)
+	{
+		return;
+	}
+
+	FScopedTransaction BakeTransform(LOCTEXT("BakeTransform", "Bake Transform"));
+
+	FocusedMovieScene->Modify();
+
+	FQualifiedFrameTime ResetTime = Sequencer->GetLocalTime();
+
+	FFrameRate TickResolution = FocusedMovieScene->GetTickResolution();
+	FFrameRate DisplayRate = FocusedMovieScene->GetDisplayRate();
+
+	FFrameTime InFrame = BakeInTime;
+	FFrameTime OutFrame = BakeOutTime;
+	FFrameTime Interval = BakeInterval;
+
+	if (Params.TimeUnit == ESequenceTimeUnit::DisplayRate)
+	{
+		InFrame = ConvertFrameTime(BakeInTime, DisplayRate, TickResolution);
+		OutFrame = ConvertFrameTime(BakeOutTime, DisplayRate, TickResolution);
+		Interval = ConvertFrameTime(BakeInterval, DisplayRate, TickResolution);
+	}
+
+	struct FBakeData
+	{
+		TArray<FVector> Locations;
+		TArray<FRotator> Rotations;
+		TArray<FVector> Scales;
+		TArray<FFrameNumber> KeyTimes;
+	};
+
+	TMap<FGuid, FBakeData> BakeDataMap;
+	for (const FSequencerBindingProxy& ObjectBinding : ObjectBindings)
+	{
+		BakeDataMap.Add(ObjectBinding.BindingID);
+	}
+
+	FMovieSceneSequenceTransform RootToLocalTransform = Sequencer->GetFocusedMovieSceneSequenceTransform();
+	
+	for (FFrameTime EvalTime = InFrame; EvalTime <= OutFrame; EvalTime += Interval)
+	{
+		FFrameNumber KeyTime = FFrameRate::Snap(EvalTime, TickResolution, DisplayRate).FloorToFrame();
+		FMovieSceneEvaluationRange Range(KeyTime * RootToLocalTransform.InverseLinearOnly(), TickResolution);
+
+		Sequencer->SetLocalTimeDirectly(Range.GetTime());
+
+		for (const FSequencerBindingProxy& ObjectBinding : ObjectBindings)
+		{
+			FGuid Guid = ObjectBinding.BindingID;
+
+			for (TWeakObjectPtr<> RuntimeObject : Sequencer->FindObjectsInCurrentSequence(Guid) )
+			{
+				AActor* Actor = Cast<AActor>(RuntimeObject.Get());
+				if (!Actor)
+				{
+					UActorComponent* ActorComponent = Cast<UActorComponent>(RuntimeObject.Get());
+					if (ActorComponent)
+					{
+						Actor = ActorComponent->GetOwner();
+					}
+				}
+
+				if (!Actor)
+				{
+					continue;
+				}
+
+				UCameraComponent* CameraComponent = MovieSceneHelpers::CameraComponentFromRuntimeObject(RuntimeObject.Get());
+
+				// Cache transforms
+				USceneComponent* Parent = nullptr;
+				if (CameraComponent)
+				{
+					Parent = CameraComponent->GetAttachParent();
+				} 
+				else if (Actor->GetRootComponent())
+				{
+					Parent = Actor->GetRootComponent()->GetAttachParent();
+				}
+				
+				// The CameraRig_rail updates the spline position tick, so it needs to be ticked manually while baking the frames
+				while (Parent && Parent->GetOwner())
+				{
+					Parent->GetOwner()->Tick(0.03f);
+					if (USkeletalMeshComponent* SkeletalMeshComponent = Cast<USkeletalMeshComponent>(Parent))
+					{
+						SkeletalMeshComponent->TickAnimation(0.f, false);
+
+						SkeletalMeshComponent->RefreshBoneTransforms();
+						SkeletalMeshComponent->RefreshSlaveComponents();
+						SkeletalMeshComponent->UpdateComponentToWorld();
+						SkeletalMeshComponent->FinalizeBoneTransform();
+						SkeletalMeshComponent->MarkRenderTransformDirty();
+						SkeletalMeshComponent->MarkRenderDynamicDataDirty();
+					}
+					Parent = Parent->GetAttachParent();
+				}
+
+				if (CameraComponent)
+				{
+					FTransform AdditiveOffset;
+					float AdditiveFOVOffset;
+					CameraComponent->GetAdditiveOffset(AdditiveOffset, AdditiveFOVOffset);
+
+					FTransform Transform(Actor->GetActorRotation(), Actor->GetActorLocation());
+					FTransform TransformWithAdditiveOffset = AdditiveOffset * Transform;
+					FVector LocalTranslation = TransformWithAdditiveOffset.GetTranslation();
+					FRotator LocalRotation = TransformWithAdditiveOffset.GetRotation().Rotator();
+
+					BakeDataMap[Guid].Locations.Add(LocalTranslation);
+					BakeDataMap[Guid].Rotations.Add(LocalRotation);
+					BakeDataMap[Guid].Scales.Add(FVector::OneVector);
+				}
+				else
+				{
+					BakeDataMap[Guid].Locations.Add(Actor->GetActorLocation());
+					BakeDataMap[Guid].Rotations.Add(Actor->GetActorRotation());
+					BakeDataMap[Guid].Scales.Add(Actor->GetActorScale());
+				}
+
+				BakeDataMap[Guid].KeyTimes.Add(KeyTime);
+			}
+		}
+	}
+
+	const bool bDisableSectionsAfterBaking = Sequencer->GetSequencerSettings()->GetDisableSectionsAfterBaking();
+
+	for (TPair<FGuid, FBakeData>& BakeData : BakeDataMap)
+	{
+		FGuid Guid = BakeData.Key;
+
+		// Disable or delete any constraint (attach/path) tracks
+		AActor* AttachParentActor = nullptr;
+		for (UMovieSceneTrack* Track : FocusedMovieScene->FindTracks(UMovieScene3DConstraintTrack::StaticClass(), Guid))
+		{
+			if (UMovieScene3DConstraintTrack* ConstraintTrack = Cast<UMovieScene3DConstraintTrack>(Track))
+			{
+				for (UMovieSceneSection* ConstraintSection : ConstraintTrack->GetAllSections())
+				{
+					FMovieSceneObjectBindingID ConstraintBindingID = (Cast<UMovieScene3DConstraintSection>(ConstraintSection))->GetConstraintBindingID();
+					for (TWeakObjectPtr<> ParentObject : ConstraintBindingID.ResolveBoundObjects(Sequencer->GetFocusedTemplateID(), *Sequencer))
+					{
+						AttachParentActor = Cast<AActor>(ParentObject.Get());
+						break;
+					}
+				}
+
+				if (bDisableSectionsAfterBaking)
+				{
+					for (UMovieSceneSection* ConstraintSection : ConstraintTrack->GetAllSections())
+					{
+						ConstraintSection->Modify();
+						ConstraintSection->SetIsActive(false);
+					}
+				}
+				else
+				{
+					FocusedMovieScene->RemoveTrack(*ConstraintTrack);
+				}
+			}
+		}
+
+		// Disable or delete any transform tracks
+		for (UMovieSceneTrack* Track : FocusedMovieScene->FindTracks(UMovieScene3DTransformTrack::StaticClass(), Guid))
+		{
+			if (UMovieScene3DTransformTrack* TransformTrack = Cast<UMovieScene3DTransformTrack>(Track))
+			{
+				if (bDisableSectionsAfterBaking)
+				{
+					for (UMovieSceneSection* TransformSection : TransformTrack->GetAllSections())
+					{
+						TransformSection->Modify();
+						TransformSection->SetIsActive(false);
+					}
+				}
+				else
+				{
+					FocusedMovieScene->RemoveTrack(*TransformTrack);
+				}
+			}
+		}
+
+		// Disable or delete any camera anim tracks
+		for (UMovieSceneTrack* Track : FocusedMovieScene->FindTracks(UMovieSceneCameraAnimTrack::StaticClass(), Guid))
+		{
+			if (UMovieSceneCameraAnimTrack* CameraAnimTrack = Cast<UMovieSceneCameraAnimTrack>(Track))
+			{
+				if (bDisableSectionsAfterBaking)
+				{
+					for (UMovieSceneSection* CameraAnimSection : CameraAnimTrack->GetAllSections())
+					{
+						CameraAnimSection->Modify();
+						CameraAnimSection->SetIsActive(false);
+					}
+				}
+				else
+				{
+					FocusedMovieScene->RemoveTrack(*CameraAnimTrack);
+				}
+			}
+		}
+
+		// Disable or delete any camera shake tracks
+		for (UMovieSceneTrack* Track : FocusedMovieScene->FindTracks(UMovieSceneCameraShakeTrack::StaticClass(), Guid))
+		{
+			if (UMovieSceneCameraShakeTrack* CameraShakeTrack = Cast<UMovieSceneCameraShakeTrack>(Track))
+			{
+				if (bDisableSectionsAfterBaking)
+				{
+					for (UMovieSceneSection* CameraShakeSection : CameraShakeTrack->GetAllSections())
+					{
+						CameraShakeSection->Modify();
+						CameraShakeSection->SetIsActive(false);
+					}
+				}
+				else
+				{
+					FocusedMovieScene->RemoveTrack(*CameraShakeTrack);
+				}
+			}
+		}
+
+		// Reset position
+		Sequencer->SetLocalTimeDirectly(ResetTime.Time);
+
+		FVector DefaultLocation = FVector::ZeroVector;
+		FVector DefaultRotation = FVector::ZeroVector;
+		FVector DefaultScale = FVector::OneVector;
+
+		for (TWeakObjectPtr<> RuntimeObject : Sequencer->FindObjectsInCurrentSequence(Guid))
+		{
+			AActor* Actor = Cast<AActor>(RuntimeObject.Get());
+			if (!Actor)
+			{
+				UActorComponent* ActorComponent = Cast<UActorComponent>(RuntimeObject.Get());
+				if (ActorComponent)
+				{
+					Actor = ActorComponent->GetOwner();
+				}
+			}
+
+			if (!Actor)
+			{
+				continue;
+			}
+
+			DefaultLocation = Actor->GetActorLocation();
+			DefaultRotation = Actor->GetActorRotation().Euler();
+			DefaultScale = Actor->GetActorScale();
+
+			// Always detach from any existing parent
+			Actor->DetachFromActor(FDetachmentTransformRules::KeepRelativeTransform);
+		}
+			
+		// Create new transform track and section
+		UMovieScene3DTransformTrack* TransformTrack = Cast<UMovieScene3DTransformTrack>(FocusedMovieScene->AddTrack(UMovieScene3DTransformTrack::StaticClass(), Guid));
+
+		if (TransformTrack)
+		{
+			UMovieScene3DTransformSection* TransformSection = CastChecked<UMovieScene3DTransformSection>(TransformTrack->CreateNewSection());
+			TransformTrack->AddSection(*TransformSection);
+
+			TransformSection->SetRange(TRange<FFrameNumber>::All());
+
+			TArrayView<FMovieSceneDoubleChannel*> DoubleChannels = TransformSection->GetChannelProxy().GetChannels<FMovieSceneDoubleChannel>();
+			DoubleChannels[0]->SetDefault(DefaultLocation.X);
+			DoubleChannels[1]->SetDefault(DefaultLocation.Y);
+			DoubleChannels[2]->SetDefault(DefaultLocation.Z);
+			DoubleChannels[3]->SetDefault(DefaultRotation.X);
+			DoubleChannels[4]->SetDefault(DefaultRotation.Y);
+			DoubleChannels[5]->SetDefault(DefaultRotation.Z);
+			DoubleChannels[6]->SetDefault(DefaultScale.X);
+			DoubleChannels[7]->SetDefault(DefaultScale.Y);
+			DoubleChannels[8]->SetDefault(DefaultScale.Z);
+
+			TArray<FVector> LocalTranslations, LocalRotations, LocalScales;
+			LocalTranslations.SetNum(BakeData.Value.KeyTimes.Num());
+			LocalRotations.SetNum(BakeData.Value.KeyTimes.Num());
+			LocalScales.SetNum(BakeData.Value.KeyTimes.Num());
+
+			for (int32 Counter = 0; Counter < BakeData.Value.KeyTimes.Num(); ++Counter)
+			{
+				FTransform LocalTransform(BakeData.Value.Rotations[Counter], BakeData.Value.Locations[Counter], BakeData.Value.Scales[Counter]);
+				LocalTranslations[Counter] = LocalTransform.GetTranslation();
+				LocalRotations[Counter] = LocalTransform.GetRotation().Euler();
+				LocalScales[Counter] = LocalTransform.GetScale3D();
+			}
+
+			// Euler filter
+			for (int32 Counter = 0; Counter < LocalRotations.Num() - 1; ++Counter)
+			{
+				FMath::WindRelativeAnglesDegrees(LocalRotations[Counter].X, LocalRotations[Counter + 1].X);
+				FMath::WindRelativeAnglesDegrees(LocalRotations[Counter].Y, LocalRotations[Counter + 1].Y);
+				FMath::WindRelativeAnglesDegrees(LocalRotations[Counter].Z, LocalRotations[Counter + 1].Z);							
+			}
+				
+			for (int32 Counter = 0; Counter < BakeData.Value.KeyTimes.Num(); ++Counter)
+			{
+				FFrameNumber KeyTime = BakeData.Value.KeyTimes[Counter];
+				DoubleChannels[0]->AddLinearKey(KeyTime, LocalTranslations[Counter].X);
+				DoubleChannels[1]->AddLinearKey(KeyTime, LocalTranslations[Counter].Y);
+				DoubleChannels[2]->AddLinearKey(KeyTime, LocalTranslations[Counter].Z);
+				DoubleChannels[3]->AddLinearKey(KeyTime, LocalRotations[Counter].X);
+				DoubleChannels[4]->AddLinearKey(KeyTime, LocalRotations[Counter].Y);
+				DoubleChannels[5]->AddLinearKey(KeyTime, LocalRotations[Counter].Z);
+				DoubleChannels[6]->AddLinearKey(KeyTime, LocalScales[Counter].X);
+				DoubleChannels[7]->AddLinearKey(KeyTime, LocalScales[Counter].Y);
+				DoubleChannels[8]->AddLinearKey(KeyTime, LocalScales[Counter].Z);
+			}
+		}
+	}
+	
+	Sequencer->NotifyMovieSceneDataChanged( EMovieSceneDataChangeType::MovieSceneStructureItemsChanged );
+}
+
 void ULevelSequenceEditorSubsystem::FixActorReferences()
 {
 	TSharedPtr<ISequencer> Sequencer = GetActiveSequencer();
@@ -86,14 +485,20 @@ void ULevelSequenceEditorSubsystem::FixActorReferences()
 		return;
 	}
 
-	FScopedTransaction FixActorReferencesTransaction(NSLOCTEXT("LevelSequenceEditor", "FixActorReferences", "Fix Actor References"));
-
 	UMovieScene* FocusedMovieScene = Sequencer->GetFocusedMovieSceneSequence()->GetMovieScene();
 	if (!FocusedMovieScene)
 	{
 		return;
 	}
 
+	if (FocusedMovieScene->IsReadOnly())
+	{
+		FSequencerUtilities::ShowReadOnlyError();
+		return;
+	}
+
+	FScopedTransaction FixActorReferencesTransaction(LOCTEXT("FixActorReferences", "Fix Actor References"));
+
 	TMap<FString, AActor*> ActorNameToActorMap;
 
 	for (TActorIterator<AActor> ActorItr(PlaybackContext); ActorItr; ++ActorItr)
@@ -154,3 +559,4 @@ void ULevelSequenceEditorSubsystem::FixActorReferences()
 	}
 }
 
+#undef LOCTEXT_NAMESPACE
diff --git a/Engine/Plugins/MovieScene/LevelSequenceEditor/Source/LevelSequenceEditor/Public/LevelSequenceEditorSubsystem.h b/Engine/Plugins/MovieScene/LevelSequenceEditor/Source/LevelSequenceEditor/Public/LevelSequenceEditorSubsystem.h
--- a/Engine/Plugins/MovieScene/LevelSequenceEditor/Source/LevelSequenceEditor/Public/LevelSequenceEditorSubsystem.h
+++ b/Engine/Plugins/MovieScene/LevelSequenceEditor/Source/LevelSequenceEditor/Public/LevelSequenceEditorSubsystem.h
@@ -7,6 +7,9 @@
 #include "EditorSubsystem.h"
 #include "ISequencer.h"
 
+#include "SequenceTimeUnit.h"
+#include "SequencerBindingProxy.h"
+
 #include "LevelSequenceEditorSubsystem.generated.h"
 
 DECLARE_LOG_CATEGORY_EXTERN(LogLevelSequenceEditor, Log, All);
@@ -14,6 +17,16 @@ DECLARE_LOG_CATEGORY_EXTERN(LogLevelSequenceEditor, Log, All);
 class FExtender;
 class UMovieSceneCompiledDataManager;
 
+USTRUCT(BlueprintType)
+struct FMovieSceneScriptingParams
+{
+	GENERATED_BODY()
+		
+	FMovieSceneScriptingParams() {}
+
+	ESequenceTimeUnit TimeUnit = ESequenceTimeUnit::DisplayRate;
+};
+
 /**
 * ULevelSequenceEditorSubsystem
 * Subsystem for level sequencer related utilities to scripts
@@ -29,13 +42,20 @@ public:
 	virtual void Deinitialize() override;
 
 	void OnSequencerCreated(TSharedRef<ISequencer> InSequencer);
+	
+	/** Bake transform */
+	UFUNCTION(BlueprintCallable, Category = "Level Sequence Editor")
+	void BakeTransform(const TArray<FSequencerBindingProxy>& ObjectBindings, const FFrameTime& BakeInTime, const FFrameTime& BakeOutTime, const FFrameTime& BakeInterval, const FMovieSceneScriptingParams& Params = FMovieSceneScriptingParams());
 
-	/** Attempts to automatically fix up broken actor references in the current scene. */
+	/** Attempts to automatically fix up broken actor references in the current scene */
 	UFUNCTION(BlueprintCallable, Category = "Level Sequence Editor")
 	void FixActorReferences();
 
 private:
+
 	TSharedPtr<ISequencer> GetActiveSequencer();
+	
+	void BakeTransformInternal();
 
 	FDelegateHandle OnSequencerCreatedHandle;
 
@@ -44,5 +64,6 @@ private:
 
 	TSharedPtr<FUICommandList> CommandList;
 
+	TSharedPtr<FExtender> BakeTransformMenuExtender;
 	TSharedPtr<FExtender> FixActorReferencesMenuExtender;
 };
diff --git a/Engine/Source/Editor/Sequencer/Private/SSequencer.cpp b/Engine/Source/Editor/Sequencer/Private/SSequencer.cpp
--- a/Engine/Source/Editor/Sequencer/Private/SSequencer.cpp
+++ b/Engine/Source/Editor/Sequencer/Private/SSequencer.cpp
@@ -2020,7 +2020,6 @@ TSharedRef<SWidget> SSequencer::MakeActionsMenu()
 
 	if (SequencerPtr.Pin()->IsLevelEditorSequencer())
 	{
-		MenuBuilder.AddMenuEntry(FSequencerCommands::Get().BakeTransform);
 		MenuBuilder.AddMenuEntry(FSequencerCommands::Get().SyncSectionsUsingSourceTimecode);
 	}
 
@@ -2047,7 +2046,6 @@ void SSequencer::FillAdvancedMenu(FMenuBuilder& MenuBuilder)
 	{
 		MenuBuilder.BeginSection("Bindings", LOCTEXT("BindingsMenuHeader", "Bindings"));
 
-		MenuBuilder.AddMenuEntry(FSequencerCommands::Get().FixPossessableObjectClass);
 		MenuBuilder.AddMenuEntry(FSequencerCommands::Get().RebindPossessableReferences);
 
 		MenuBuilder.EndSection();
diff --git a/Engine/Source/Editor/Sequencer/Private/Sequencer.cpp b/Engine/Source/Editor/Sequencer/Private/Sequencer.cpp
--- a/Engine/Source/Editor/Sequencer/Private/Sequencer.cpp
+++ b/Engine/Source/Editor/Sequencer/Private/Sequencer.cpp
@@ -2230,336 +2230,6 @@ bool FSequencer::IsPilotCamera() const
 	return false;
 }
 
-void FSequencer::BakeTransform()
-{
-	UMovieScene* FocusedMovieScene = GetFocusedMovieSceneSequence()->GetMovieScene();
-	if (!FocusedMovieScene)
-	{
-		return;
-	}
-
-	if (FocusedMovieScene->IsReadOnly())
-	{
-		FSequencerUtilities::ShowReadOnlyError();
-		return;
-	}
-
-	FScopedTransaction BakeTransform(NSLOCTEXT("Sequencer", "BakeTransform", "Bake Transform"));
-
-	FocusedMovieScene->Modify();
-
-	TArray<FGuid> Guids;
-	for (const TSharedRef<FSequencerDisplayNode>& Node : Selection.GetSelectedOutlinerNodes())
-	{
-		if (Node->GetType() != ESequencerNode::Object)
-		{
-			continue;
-		}
-
-		auto ObjectBindingNode = StaticCastSharedRef<FSequencerObjectBindingNode>(Node);
-		FGuid Guid = ObjectBindingNode->GetObjectBinding();
-
-		Guids.Add(Guid);
-	}
-
-	FFrameTime ResetTime = PlayPosition.GetCurrentPosition();
-
-	FFrameRate   Resolution = FocusedMovieScene->GetTickResolution();
-	FFrameRate   SnapRate = FocusedMovieScene->GetDisplayRate();
-
-	FFrameNumber InFrame = UE::MovieScene::DiscreteInclusiveLower(GetPlaybackRange());
-	FFrameNumber OutFrame = UE::MovieScene::DiscreteExclusiveUpper(GetPlaybackRange());
-
-	struct FBakeData
-	{
-		TArray<FVector> Locations;
-		TArray<FRotator> Rotations;
-		TArray<FVector> Scales;
-		TArray<FFrameNumber> KeyTimes;
-	};
-
-	TMap<FGuid, FBakeData> BakeDataMap;
-	for (FGuid Guid : Guids)
-	{
-		BakeDataMap.Add(Guid);
-	}
-
-	FFrameTime Interval = FFrameRate::TransformTime(1, SnapRate, Resolution);
-	for (FFrameTime EvalTime = InFrame; EvalTime < OutFrame; EvalTime += Interval)
-	{
-		FFrameNumber KeyTime = FFrameRate::Snap(EvalTime, Resolution, SnapRate).FloorToFrame();
-		FMovieSceneEvaluationRange Range(KeyTime * RootToLocalTransform.InverseLinearOnly(), Resolution);
-
-		EvaluateInternal(Range);
-
-		for (FGuid Guid : Guids)
-		{
-			for (auto RuntimeObject : FindBoundObjects(Guid, ActiveTemplateIDs.Top()) )
-			{
-				AActor* Actor = Cast<AActor>(RuntimeObject.Get());
-				if (!Actor)
-				{
-					UActorComponent* ActorComponent = Cast<UActorComponent>(RuntimeObject.Get());
-					if (ActorComponent)
-					{
-						Actor = ActorComponent->GetOwner();
-					}
-				}
-
-				if (!Actor)
-				{
-					continue;
-				}
-
-				UCameraComponent* CameraComponent = MovieSceneHelpers::CameraComponentFromRuntimeObject(RuntimeObject.Get());
-
-				// Cache transforms
-				USceneComponent* Parent = nullptr;
-				if (CameraComponent)
-				{
-					Parent = CameraComponent->GetAttachParent();
-				} 
-				else if (Actor->GetRootComponent())
-				{
-					Parent = Actor->GetRootComponent()->GetAttachParent();
-				}
-				
-				// The CameraRig_rail updates the spline position tick, so it needs to be ticked manually while baking the frames
-				while (Parent && Parent->GetOwner())
-				{
-					Parent->GetOwner()->Tick(0.03f);
-					if (USkeletalMeshComponent* SkeletalMeshComponent = Cast<USkeletalMeshComponent>(Parent))
-					{
-						SkeletalMeshComponent->TickAnimation(0.f, false);
-
-						SkeletalMeshComponent->RefreshBoneTransforms();
-						SkeletalMeshComponent->RefreshSlaveComponents();
-						SkeletalMeshComponent->UpdateComponentToWorld();
-						SkeletalMeshComponent->FinalizeBoneTransform();
-						SkeletalMeshComponent->MarkRenderTransformDirty();
-						SkeletalMeshComponent->MarkRenderDynamicDataDirty();
-					}
-					Parent = Parent->GetAttachParent();
-				}
-
-				if (CameraComponent)
-				{
-					FTransform AdditiveOffset;
-					float AdditiveFOVOffset;
-					CameraComponent->GetAdditiveOffset(AdditiveOffset, AdditiveFOVOffset);
-
-					FTransform Transform(Actor->GetActorRotation(), Actor->GetActorLocation());
-					FTransform TransformWithAdditiveOffset = AdditiveOffset * Transform;
-					FVector LocalTranslation = TransformWithAdditiveOffset.GetTranslation();
-					FRotator LocalRotation = TransformWithAdditiveOffset.GetRotation().Rotator();
-
-					BakeDataMap[Guid].Locations.Add(LocalTranslation);
-					BakeDataMap[Guid].Rotations.Add(LocalRotation);
-					BakeDataMap[Guid].Scales.Add(FVector::OneVector);
-				}
-				else
-				{
-					BakeDataMap[Guid].Locations.Add(Actor->GetActorLocation());
-					BakeDataMap[Guid].Rotations.Add(Actor->GetActorRotation());
-					BakeDataMap[Guid].Scales.Add(Actor->GetActorScale());
-				}
-
-				BakeDataMap[Guid].KeyTimes.Add(KeyTime);
-			}
-		}
-	}
-
-	bool bDisableSectionsAfterBaking = Settings->GetDisableSectionsAfterBaking();
-
-	for (auto& BakeData : BakeDataMap)
-	{
-		FGuid Guid = BakeData.Key;
-
-		// Disable or delete any attach tracks
-		AActor* AttachParentActor = nullptr;
-		for (UMovieSceneTrack* Track : FocusedMovieScene->FindTracks(UMovieScene3DAttachTrack::StaticClass(), Guid))
-		{
-			if (UMovieScene3DAttachTrack* AttachTrack = Cast<UMovieScene3DAttachTrack>(Track))
-			{
-				for (auto AttachSection : AttachTrack->GetAllSections())
-				{
-					FMovieSceneObjectBindingID ConstraintBindingID = (Cast<UMovieScene3DAttachSection>(AttachSection))->GetConstraintBindingID();
-					for (auto ParentObject : ConstraintBindingID.ResolveBoundObjects(GetFocusedTemplateID(), *this))
-					{
-						AttachParentActor = Cast<AActor>(ParentObject.Get());
-						break;
-					}
-				}
-
-				if (bDisableSectionsAfterBaking)
-				{
-					for (auto AttachSection : AttachTrack->GetAllSections())
-					{
-						AttachSection->Modify();
-						AttachSection->SetIsActive(false);
-					}
-				}
-				else
-				{
-					FocusedMovieScene->RemoveTrack(*AttachTrack);
-				}
-			}
-		}
-
-		// Disable or delete any transform tracks
-		for (UMovieSceneTrack* Track : FocusedMovieScene->FindTracks(UMovieScene3DTransformTrack::StaticClass(), Guid))
-		{
-			if (UMovieScene3DTransformTrack* TransformTrack = Cast<UMovieScene3DTransformTrack>(Track))
-			{
-				if (bDisableSectionsAfterBaking)
-				{
-					for (auto TransformSection : TransformTrack->GetAllSections())
-					{
-						TransformSection->Modify();
-						TransformSection->SetIsActive(false);
-					}
-				}
-				else
-				{
-					FocusedMovieScene->RemoveTrack(*TransformTrack);
-				}
-			}
-		}
-
-		// Disable or delete any camera anim tracks
-		for (UMovieSceneTrack* Track : FocusedMovieScene->FindTracks(UMovieSceneCameraAnimTrack::StaticClass(), Guid))
-		{
-			if (UMovieSceneCameraAnimTrack* CameraAnimTrack = Cast<UMovieSceneCameraAnimTrack>(Track))
-			{
-				if (bDisableSectionsAfterBaking)
-				{
-					for (auto CameraAnimSection : CameraAnimTrack->GetAllSections())
-					{
-						CameraAnimSection->Modify();
-						CameraAnimSection->SetIsActive(false);
-					}
-				}
-				else
-				{
-					FocusedMovieScene->RemoveTrack(*CameraAnimTrack);
-				}
-			}
-		}
-
-		// Disable or delete any camera shake tracks
-		for (UMovieSceneTrack* Track : FocusedMovieScene->FindTracks(UMovieSceneCameraShakeTrack::StaticClass(), Guid))
-		{
-			if (UMovieSceneCameraShakeTrack* CameraShakeTrack = Cast<UMovieSceneCameraShakeTrack>(Track))
-			{
-				if (bDisableSectionsAfterBaking)
-				{
-					for (auto CameraShakeSection : CameraShakeTrack->GetAllSections())
-					{
-						CameraShakeSection->Modify();
-						CameraShakeSection->SetIsActive(false);
-					}
-				}
-				else
-				{
-					FocusedMovieScene->RemoveTrack(*CameraShakeTrack);
-				}
-			}
-		}
-
-		// Reset position
-		EvaluateInternal(PlayPosition.JumpTo(ResetTime));
-
-		FVector DefaultLocation = FVector::ZeroVector;
-		FVector DefaultRotation = FVector::ZeroVector;
-		FVector DefaultScale = FVector::OneVector;
-
-		for (auto RuntimeObject : FindBoundObjects(Guid, ActiveTemplateIDs.Top()))
-		{
-			AActor* Actor = Cast<AActor>(RuntimeObject.Get());
-			if (!Actor)
-			{
-				UActorComponent* ActorComponent = Cast<UActorComponent>(RuntimeObject.Get());
-				if (ActorComponent)
-				{
-					Actor = ActorComponent->GetOwner();
-				}
-			}
-
-			if (!Actor)
-			{
-				continue;
-			}
-
-			DefaultLocation = Actor->GetActorLocation();
-			DefaultRotation = Actor->GetActorRotation().Euler();
-			DefaultScale = Actor->GetActorScale();
-
-			// Always detach from any existing parent
-			Actor->DetachFromActor(FDetachmentTransformRules::KeepRelativeTransform);
-		}
-			
-		// Create new transform track and section
-		UMovieScene3DTransformTrack* TransformTrack = Cast<UMovieScene3DTransformTrack>(FocusedMovieScene->AddTrack(UMovieScene3DTransformTrack::StaticClass(), Guid));
-
-		if (TransformTrack)
-		{
-			UMovieScene3DTransformSection* TransformSection = CastChecked<UMovieScene3DTransformSection>(TransformTrack->CreateNewSection());
-			TransformTrack->AddSection(*TransformSection);
-
-			TransformSection->SetRange(TRange<FFrameNumber>::All());
-
-			TArrayView<FMovieSceneDoubleChannel*> DoubleChannels = TransformSection->GetChannelProxy().GetChannels<FMovieSceneDoubleChannel>();
-			DoubleChannels[0]->SetDefault(DefaultLocation.X);
-			DoubleChannels[1]->SetDefault(DefaultLocation.Y);
-			DoubleChannels[2]->SetDefault(DefaultLocation.Z);
-			DoubleChannels[3]->SetDefault(DefaultRotation.X);
-			DoubleChannels[4]->SetDefault(DefaultRotation.Y);
-			DoubleChannels[5]->SetDefault(DefaultRotation.Z);
-			DoubleChannels[6]->SetDefault(DefaultScale.X);
-			DoubleChannels[7]->SetDefault(DefaultScale.Y);
-			DoubleChannels[8]->SetDefault(DefaultScale.Z);
-
-			TArray<FVector> LocalTranslations, LocalRotations, LocalScales;
-			LocalTranslations.SetNum(BakeData.Value.KeyTimes.Num());
-			LocalRotations.SetNum(BakeData.Value.KeyTimes.Num());
-			LocalScales.SetNum(BakeData.Value.KeyTimes.Num());
-
-			for (int32 Counter = 0; Counter < BakeData.Value.KeyTimes.Num(); ++Counter)
-			{
-				FTransform LocalTransform(BakeData.Value.Rotations[Counter], BakeData.Value.Locations[Counter], BakeData.Value.Scales[Counter]);
-				LocalTranslations[Counter] = LocalTransform.GetTranslation();
-				LocalRotations[Counter] = LocalTransform.GetRotation().Euler();
-				LocalScales[Counter] = LocalTransform.GetScale3D();
-			}
-
-			// Euler filter
-			for (int32 Counter = 0; Counter < LocalRotations.Num() - 1; ++Counter)
-			{
-				FMath::WindRelativeAnglesDegrees(LocalRotations[Counter].X, LocalRotations[Counter + 1].X);
-				FMath::WindRelativeAnglesDegrees(LocalRotations[Counter].Y, LocalRotations[Counter + 1].Y);
-				FMath::WindRelativeAnglesDegrees(LocalRotations[Counter].Z, LocalRotations[Counter + 1].Z);							
-			}
-				
-			for (int32 Counter = 0; Counter < BakeData.Value.KeyTimes.Num(); ++Counter)
-			{
-				FFrameNumber KeyTime = BakeData.Value.KeyTimes[Counter];
-				DoubleChannels[0]->AddLinearKey(KeyTime, LocalTranslations[Counter].X);
-				DoubleChannels[1]->AddLinearKey(KeyTime, LocalTranslations[Counter].Y);
-				DoubleChannels[2]->AddLinearKey(KeyTime, LocalTranslations[Counter].Z);
-				DoubleChannels[3]->AddLinearKey(KeyTime, LocalRotations[Counter].X);
-				DoubleChannels[4]->AddLinearKey(KeyTime, LocalRotations[Counter].Y);
-				DoubleChannels[5]->AddLinearKey(KeyTime, LocalRotations[Counter].Z);
-				DoubleChannels[6]->AddLinearKey(KeyTime, LocalScales[Counter].X);
-				DoubleChannels[7]->AddLinearKey(KeyTime, LocalScales[Counter].Y);
-				DoubleChannels[8]->AddLinearKey(KeyTime, LocalScales[Counter].Z);
-			}
-		}
-	}
-	
-	NotifyMovieSceneDataChanged( EMovieSceneDataChangeType::MovieSceneStructureItemsChanged );
-}
-
-
 void FSequencer::SyncSectionsUsingSourceTimecode()
 {
 	FScopedTransaction SyncSectionsUsingSourceTimecodeTransaction( LOCTEXT("SyncSectionsUsingSourceTimecode_Transaction", "Sync Sections Using Source Timecode") );
@@ -13353,11 +13023,6 @@ void FSequencer::BindCommands()
 		FIsActionButtonVisible::CreateLambda([this] { return ExactCast<ULevelSequence>(GetFocusedMovieSceneSequence()) != nullptr && IVREditorModule::Get().IsVREditorModeActive() == false; }) //@todo VREditor: Creating a camera while in VR mode disrupts the hmd. This is a temporary fix by hiding the button when in VR mode.
 	);
 
-	SequencerCommandBindings->MapAction(
-		Commands.BakeTransform,
-		FExecuteAction::CreateSP( this, &FSequencer::BakeTransform ),
-		FCanExecuteAction::CreateLambda( []{ return true; } ) );
-
 	SequencerCommandBindings->MapAction(
 		Commands.SyncSectionsUsingSourceTimecode,
 		FExecuteAction::CreateSP( this, &FSequencer::SyncSectionsUsingSourceTimecode ),
diff --git a/Engine/Source/Editor/Sequencer/Private/Sequencer.h b/Engine/Source/Editor/Sequencer/Private/Sequencer.h
--- a/Engine/Source/Editor/Sequencer/Private/Sequencer.h
+++ b/Engine/Source/Editor/Sequencer/Private/Sequencer.h
@@ -264,9 +264,6 @@ public:
 	/** Shrink time*/
 	void ShrinkTime(FFrameTime InDeltaTime);
 
-	/** Bake transform */
-	void BakeTransform();
-
 	/** Sync using source timecode */
 	void SyncSectionsUsingSourceTimecode();
 
diff --git a/Engine/Source/Editor/Sequencer/Private/SequencerCommands.cpp b/Engine/Source/Editor/Sequencer/Private/SequencerCommands.cpp
--- a/Engine/Source/Editor/Sequencer/Private/SequencerCommands.cpp
+++ b/Engine/Source/Editor/Sequencer/Private/SequencerCommands.cpp
@@ -111,7 +111,6 @@ void FSequencerCommands::RegisterCommands()
 	UI_COMMAND( OpenDirectorBlueprint, "Open Director Blueprint", "Opens the director blueprint for this sequence.", EUserInterfaceActionType::Button, FInputChord() );
 	UI_COMMAND( OpenTaggedBindingManager, "Open Binding Tag Manager", "Specifies options for tagging bindings within this sequence for external systems to reference as a persistent name.", EUserInterfaceActionType::Button, FInputChord() );
 	UI_COMMAND( OpenNodeGroupsManager, "Open Sequencer Group Manager", "Manage groups within this sequence.", EUserInterfaceActionType::Button, FInputChord());
-	UI_COMMAND( BakeTransform, "Bake Transform", "Bake transform in world space, removing any existing transform and attach tracks", EUserInterfaceActionType::Button, FInputChord() );
 	UI_COMMAND( SyncSectionsUsingSourceTimecode, "Sync Sections using Source Timecode", "Synchronize sections to the first selected section using source timecode", EUserInterfaceActionType::Button, FInputChord() );
 	UI_COMMAND( ToggleShowRangeSlider, "Range Slider", "Enables and disables showing the time range slider", EUserInterfaceActionType::ToggleButton, FInputChord() );
 	UI_COMMAND( ToggleIsSnapEnabled, "Enable Snapping", "Enables and disables snapping", EUserInterfaceActionType::ToggleButton, FInputChord() );
diff --git a/Engine/Source/Editor/Sequencer/Public/SequencerCommands.h b/Engine/Source/Editor/Sequencer/Public/SequencerCommands.h
--- a/Engine/Source/Editor/Sequencer/Public/SequencerCommands.h
+++ b/Engine/Source/Editor/Sequencer/Public/SequencerCommands.h
@@ -286,9 +286,6 @@ public:
 	/** Remove selected nodes from folder. */
 	TSharedPtr< FUICommandInfo > RemoveFromFolder;
 
-	/** Bake transform. */
-	TSharedPtr< FUICommandInfo > BakeTransform;
-
 	/** Sync sections using source timecode. */
 	TSharedPtr< FUICommandInfo > SyncSectionsUsingSourceTimecode;
 
