Sequencer: Check for valid spawned actor instance before dereferencing. This fixes a crash if selecting an actor base class.

diff --git a/Engine/Plugins/MovieScene/LevelSequenceEditor/Source/LevelSequenceEditor/Private/Misc/LevelSequenceEditorActorSpawner.cpp b/Engine/Plugins/MovieScene/LevelSequenceEditor/Source/LevelSequenceEditor/Private/Misc/LevelSequenceEditorActorSpawner.cpp
--- a/Engine/Plugins/MovieScene/LevelSequenceEditor/Source/LevelSequenceEditor/Private/Misc/LevelSequenceEditorActorSpawner.cpp
+++ b/Engine/Plugins/MovieScene/LevelSequenceEditor/Source/LevelSequenceEditor/Private/Misc/LevelSequenceEditorActorSpawner.cpp
@@ -116,12 +116,15 @@ TValueOrError<FNewSpawnable, FText> FLevelSequenceEditorActorSpawner::CreateNewS
 				SpawnParams.Name = ActorName;
 
 				AActor* Instance = FactoryToUse->CreateActor(&SourceObject, World->PersistentLevel, FTransform(), SpawnParams);
-				Instance->bIsEditorPreviewActor = false;
-				NewSpawnable.ObjectTemplate = StaticDuplicateObject(Instance, &OwnerMovieScene, TemplateName, RF_AllFlags & ~RF_Transient);
+				if (Instance)
+				{
+					Instance->bIsEditorPreviewActor = false;
+					NewSpawnable.ObjectTemplate = StaticDuplicateObject(Instance, &OwnerMovieScene, TemplateName, RF_AllFlags & ~RF_Transient);
 
-				const bool bNetForce = false;
-				const bool bShouldModifyLevel = false;
-				World->DestroyActor(Instance, bNetForce, bShouldModifyLevel);
+					const bool bNetForce = false;
+					const bool bShouldModifyLevel = false;
+					World->DestroyActor(Instance, bNetForce, bShouldModifyLevel);
+				}
 			}
 		}
 	}
