Sequencer: Move AssignActor to LevelSequenceEditorSubSystem and expose AddActorsToBinding, ReplaceBindingWithActors, RemoveActorsFromBinding, RemoveAllBindings, RemoveInvalidBindings to python/blueprints

FSequencer::AddActorsToBinding, FSEquencer::ReplaceBindingWithActors, FSequencer::RemoveActorsFromBinding can be removed once PasteObjectBindings is moved to LevelSequenceEditorSubsystem for exposing copy/paste to python/blueprints.

diff --git a/Engine/Plugins/MovieScene/LevelSequenceEditor/Source/LevelSequenceEditor/Private/LevelSequenceEditorCommands.cpp b/Engine/Plugins/MovieScene/LevelSequenceEditor/Source/LevelSequenceEditor/Private/LevelSequenceEditorCommands.cpp
--- a/Engine/Plugins/MovieScene/LevelSequenceEditor/Source/LevelSequenceEditor/Private/LevelSequenceEditorCommands.cpp
+++ b/Engine/Plugins/MovieScene/LevelSequenceEditor/Source/LevelSequenceEditor/Private/LevelSequenceEditorCommands.cpp
@@ -17,6 +17,12 @@ void FLevelSequenceEditorCommands::RegisterCommands()
 
 	UI_COMMAND(BakeTransform, "Bake Transform", "Bake transform in world space, removing any existing transform and attach tracks", EUserInterfaceActionType::Button, FInputChord());
 	UI_COMMAND(FixActorReferences, "Fix Actor References", "Try to automatically fix up broken actor bindings", EUserInterfaceActionType::Button, FInputChord());
+
+	UI_COMMAND(AddActorsToBinding, "Add Selected", "Add selected actors to this track", EUserInterfaceActionType::Button, FInputChord());
+	UI_COMMAND(RemoveActorsFromBinding, "Remove Selected", "Remove selected actors from this track", EUserInterfaceActionType::Button, FInputChord());
+	UI_COMMAND(ReplaceBindingWithActors, "Replace with Selected", "Replace the object binding with selected actors", EUserInterfaceActionType::Button, FInputChord());
+	UI_COMMAND(RemoveAllBindings, "Remove All", "Remove all bound actors from this track", EUserInterfaceActionType::Button, FInputChord());
+	UI_COMMAND(RemoveInvalidBindings, "Remove Missing", "Remove missing objects bound to this track", EUserInterfaceActionType::Button, FInputChord());
 }
 
 #undef LOCTEXT_NAMESPACE
diff --git a/Engine/Plugins/MovieScene/LevelSequenceEditor/Source/LevelSequenceEditor/Private/LevelSequenceEditorCommands.h b/Engine/Plugins/MovieScene/LevelSequenceEditor/Source/LevelSequenceEditor/Private/LevelSequenceEditorCommands.h
--- a/Engine/Plugins/MovieScene/LevelSequenceEditor/Source/LevelSequenceEditor/Private/LevelSequenceEditorCommands.h
+++ b/Engine/Plugins/MovieScene/LevelSequenceEditor/Source/LevelSequenceEditor/Private/LevelSequenceEditorCommands.h
@@ -22,4 +22,10 @@ public:
 	
 	TSharedPtr<FUICommandInfo> BakeTransform;
 	TSharedPtr<FUICommandInfo> FixActorReferences;
+
+	TSharedPtr<FUICommandInfo> AddActorsToBinding;
+	TSharedPtr<FUICommandInfo> RemoveActorsFromBinding;
+	TSharedPtr<FUICommandInfo> ReplaceBindingWithActors;
+	TSharedPtr<FUICommandInfo> RemoveAllBindings;
+	TSharedPtr<FUICommandInfo> RemoveInvalidBindings;
 };
diff --git a/Engine/Plugins/MovieScene/LevelSequenceEditor/Source/LevelSequenceEditor/Private/LevelSequenceEditorSubsystem.cpp b/Engine/Plugins/MovieScene/LevelSequenceEditor/Source/LevelSequenceEditor/Private/LevelSequenceEditorSubsystem.cpp
--- a/Engine/Plugins/MovieScene/LevelSequenceEditor/Source/LevelSequenceEditor/Private/LevelSequenceEditorSubsystem.cpp
+++ b/Engine/Plugins/MovieScene/LevelSequenceEditor/Source/LevelSequenceEditor/Private/LevelSequenceEditorSubsystem.cpp
@@ -25,15 +25,28 @@
 #include "Camera/CameraComponent.h"
 #include "Components/SkeletalMeshComponent.h"
 #include "EngineUtils.h"
+#include "Engine/Selection.h"
 #include "Framework/MultiBox/MultiBoxBuilder.h"
+#include "Framework/Notifications/NotificationManager.h"
+#include "GameFramework/Actor.h"
 #include "LevelEditor.h"
 #include "Modules/ModuleManager.h"
 #include "ScopedTransaction.h"
+#include "Widgets/Notifications/SNotificationList.h"
 
 DEFINE_LOG_CATEGORY(LogLevelSequenceEditor);
 
 #define LOCTEXT_NAMESPACE "LevelSequenceEditor"
 
+void AddAssignActorMenu(FMenuBuilder& MenuBuilder)
+{
+	MenuBuilder.AddMenuEntry(FLevelSequenceEditorCommands::Get().AddActorsToBinding);
+	MenuBuilder.AddMenuEntry(FLevelSequenceEditorCommands::Get().ReplaceBindingWithActors);
+	MenuBuilder.AddMenuEntry(FLevelSequenceEditorCommands::Get().RemoveActorsFromBinding);
+	MenuBuilder.AddMenuEntry(FLevelSequenceEditorCommands::Get().RemoveAllBindings);
+	MenuBuilder.AddMenuEntry(FLevelSequenceEditorCommands::Get().RemoveInvalidBindings);
+}
+
 void ULevelSequenceEditorSubsystem::Initialize(FSubsystemCollectionBase& Collection)
 {
 	UE_LOG(LogLevelSequenceEditor, Log, TEXT("LevelSequenceEditor subsystem initialized."));
@@ -41,6 +54,12 @@ void ULevelSequenceEditorSubsystem::Initialize(FSubsystemCollectionBase& Collect
 	ISequencerModule& SequencerModule = FModuleManager::Get().LoadModuleChecked<ISequencerModule>("Sequencer");
 	OnSequencerCreatedHandle = SequencerModule.RegisterOnSequencerCreated(FOnSequencerCreated::FDelegate::CreateUObject(this, &ULevelSequenceEditorSubsystem::OnSequencerCreated));
 
+	auto AreActorsSelected = [this]{
+		TArray<AActor*> SelectedActors;
+		GEditor->GetSelectedActors()->GetSelectedObjects<AActor>(SelectedActors); 
+		return SelectedActors.Num() > 0;
+	};
+
 	/* Commands for this subsystem */
 	CommandList = MakeShareable(new FUICommandList);
 
@@ -52,6 +71,25 @@ void ULevelSequenceEditorSubsystem::Initialize(FSubsystemCollectionBase& Collect
 		FExecuteAction::CreateUObject(this, &ULevelSequenceEditorSubsystem::FixActorReferences)
 	);
 
+	CommandList->MapAction(FLevelSequenceEditorCommands::Get().AddActorsToBinding,
+		FExecuteAction::CreateUObject(this, &ULevelSequenceEditorSubsystem::AddActorsToBindingInternal),
+		FCanExecuteAction::CreateLambda(AreActorsSelected));
+
+	CommandList->MapAction(FLevelSequenceEditorCommands::Get().ReplaceBindingWithActors,
+		FExecuteAction::CreateUObject(this, &ULevelSequenceEditorSubsystem::ReplaceBindingWithActorsInternal),
+		FCanExecuteAction::CreateLambda(AreActorsSelected));
+
+	CommandList->MapAction(FLevelSequenceEditorCommands::Get().RemoveActorsFromBinding,
+		FExecuteAction::CreateUObject(this, &ULevelSequenceEditorSubsystem::RemoveActorsFromBindingInternal),
+		FCanExecuteAction::CreateLambda(AreActorsSelected));
+
+	CommandList->MapAction(FLevelSequenceEditorCommands::Get().RemoveAllBindings,
+		FExecuteAction::CreateUObject(this, &ULevelSequenceEditorSubsystem::RemoveAllBindingsInternal));
+
+	CommandList->MapAction(FLevelSequenceEditorCommands::Get().RemoveInvalidBindings,
+		FExecuteAction::CreateUObject(this, &ULevelSequenceEditorSubsystem::RemoveInvalidBindingsInternal));
+
+	/* Menu extenders */
 	BakeTransformMenuExtender = MakeShareable(new FExtender);
 	BakeTransformMenuExtender->AddMenuExtension("Transform", EExtensionHook::After, CommandList, FMenuExtensionDelegate::CreateStatic([](FMenuBuilder& MenuBuilder) {
 		MenuBuilder.AddMenuEntry(FLevelSequenceEditorCommands::Get().BakeTransform);
@@ -65,6 +103,17 @@ void ULevelSequenceEditorSubsystem::Initialize(FSubsystemCollectionBase& Collect
 		}));
 
 	SequencerModule.GetActionsMenuExtensibilityManager()->AddExtender(FixActorReferencesMenuExtender);
+
+	AssignActorMenuExtender = MakeShareable(new FExtender);
+	AssignActorMenuExtender->AddMenuExtension("Possessable", EExtensionHook::First, CommandList, FMenuExtensionDelegate::CreateStatic([](FMenuBuilder& MenuBuilder) {
+		FFormatNamedArguments Args;
+		MenuBuilder.AddSubMenu(
+			FText::Format(LOCTEXT("AssignActor", "Assign Actor"), Args),
+			FText::Format(LOCTEXT("AssignActorTooltip", "Assign an actor to this track"), Args),
+			FNewMenuDelegate::CreateStatic(&AddAssignActorMenu));
+		}));
+
+	SequencerModule.GetObjectBindingContextMenuExtensibilityManager()->AddExtender(AssignActorMenuExtender);
 }
 
 void ULevelSequenceEditorSubsystem::Deinitialize()
@@ -527,16 +576,16 @@ void ULevelSequenceEditorSubsystem::FixActorReferences()
 	TMap<FGuid, FGuid> OldGuidToNewGuidMap;
 	for (const FMovieScenePossessable& ActorPossessableToFix : ActorsPossessablesToFix)
 	{
-		AActor** ActorPtr = ActorNameToActorMap.Find(ActorPossessableToFix.GetName());
+		AActor* ActorPtr = *ActorNameToActorMap.Find(ActorPossessableToFix.GetName());
 		if (ActorPtr != nullptr)
 		{
 			FGuid OldGuid = ActorPossessableToFix.GetGuid();
 
 			// The actor might have an existing guid while the possessable with the same name might not. 
 			// In that case, make sure we also replace the existing guid with the new guid 
-			FGuid ExistingGuid = Sequencer->FindObjectId(**ActorPtr, Sequencer->GetFocusedTemplateID());
+			FGuid ExistingGuid = Sequencer->FindObjectId(*ActorPtr, Sequencer->GetFocusedTemplateID());
 
-			FGuid NewGuid = FSequencerUtilities::DoAssignActor(Sequencer.Get(), ActorPtr, 1, ActorPossessableToFix.GetGuid());
+			FGuid NewGuid = FSequencerUtilities::DoAssignActor(Sequencer.Get(), ActorPtr, ActorPossessableToFix.GetGuid());
 
 			OldGuidToNewGuidMap.Add(OldGuid, NewGuid);
 
@@ -559,4 +608,429 @@ void ULevelSequenceEditorSubsystem::FixActorReferences()
 	}
 }
 
+void ULevelSequenceEditorSubsystem::AddActorsToBindingInternal()
+{
+	TSharedPtr<ISequencer> Sequencer = GetActiveSequencer();
+	if (Sequencer == nullptr)
+	{
+		return;
+	}
+
+	TArray<FGuid> ObjectBindings;
+	Sequencer->GetSelectedObjects(ObjectBindings);
+	if (ObjectBindings.Num() == 0)
+	{
+		return;
+	}
+
+	FSequencerBindingProxy BindingProxy(ObjectBindings[0], Sequencer->GetFocusedMovieSceneSequence());
+
+	TArray<AActor*> SelectedActors;
+	GEditor->GetSelectedActors()->GetSelectedObjects<AActor>(SelectedActors); 
+	AddActorsToBinding(SelectedActors, BindingProxy);
+}
+
+void ULevelSequenceEditorSubsystem::AddActorsToBinding(const TArray<AActor*>& Actors, const FSequencerBindingProxy& ObjectBinding)
+{
+	if (!Actors.Num())
+	{
+		return;
+	}
+
+	TSharedPtr<ISequencer> Sequencer = GetActiveSequencer();
+	if (Sequencer == nullptr)
+	{
+		return;
+	}
+
+	UMovieSceneSequence* Sequence = Sequencer->GetFocusedMovieSceneSequence();
+	if (!Sequence)
+	{
+		return;
+	}
+
+	UMovieScene* MovieScene = Sequence->GetMovieScene();
+	if (!MovieScene)
+	{
+		return;
+	}
+
+	UClass* ActorClass = nullptr;
+	int32 NumRuntimeObjects = 0;
+
+	FGuid Guid = ObjectBinding.BindingID;
+	TArrayView<TWeakObjectPtr<>> ObjectsInCurrentSequence = Sequencer->FindObjectsInCurrentSequence(Guid);
+
+	for (TWeakObjectPtr<> Ptr : ObjectsInCurrentSequence)
+	{
+		if (const AActor* Actor = Cast<AActor>(Ptr.Get()))
+		{
+			ActorClass = Actor->GetClass();
+			++NumRuntimeObjects;
+		}
+	}
+
+	FScopedTransaction AddSelectedToBinding(LOCTEXT("AddSelectedToBinding", "Add Selected to Binding"));
+
+	Sequence->Modify();
+	MovieScene->Modify();
+
+	// Bind objects
+	int32 NumObjectsAdded = 0;
+	for (AActor* ActorToAdd : Actors)
+	{
+		if (!ObjectsInCurrentSequence.Contains(ActorToAdd))
+		{
+			if (ActorClass == nullptr || UClass::FindCommonBase(ActorToAdd->GetClass(), ActorClass) != nullptr)
+			{
+				if (ActorClass == nullptr)
+				{
+					ActorClass = ActorToAdd->GetClass();
+				}
+
+				ActorToAdd->Modify();
+				Sequence->BindPossessableObject(Guid, *ActorToAdd, Sequencer->GetPlaybackContext());
+				++NumObjectsAdded;
+			}
+			else
+			{
+				const FText NotificationText = FText::Format(LOCTEXT("UnableToAssignObject", "Cannot assign object {0}. Expected class {1}"), FText::FromString(ActorToAdd->GetName()), FText::FromString(ActorClass->GetName()));
+				FNotificationInfo Info(NotificationText);
+				Info.ExpireDuration = 3.f;
+				Info.bUseLargeFont = false;
+				FSlateNotificationManager::Get().AddNotification(Info);
+			}
+		}
+	}
+
+	// Update label
+	if (NumRuntimeObjects + NumObjectsAdded > 0)
+	{
+		FMovieScenePossessable* Possessable = MovieScene->FindPossessable(Guid);
+		if (Possessable && ActorClass != nullptr)
+		{
+			if (NumRuntimeObjects + NumObjectsAdded > 1)
+			{
+				FString NewLabel = ActorClass->GetName() + FString::Printf(TEXT(" (%d)"), NumRuntimeObjects + NumObjectsAdded);
+				Possessable->SetName(NewLabel);
+			}
+			else if (NumObjectsAdded > 0 && Actors.Num() > 0)
+			{
+				Possessable->SetName(Actors[0]->GetActorLabel());
+			}
+
+			Possessable->SetPossessedObjectClass(ActorClass);
+		}
+	}
+
+	Sequencer->RestorePreAnimatedState();
+
+	Sequencer->NotifyMovieSceneDataChanged(EMovieSceneDataChangeType::MovieSceneStructureItemsChanged);
+}
+
+void ULevelSequenceEditorSubsystem::ReplaceBindingWithActorsInternal()
+{
+	TSharedPtr<ISequencer> Sequencer = GetActiveSequencer();
+	if (Sequencer == nullptr)
+	{
+		return;
+	}
+
+	TArray<FGuid> ObjectBindings;
+	Sequencer->GetSelectedObjects(ObjectBindings);
+	if (ObjectBindings.Num() == 0)
+	{
+		return;
+	}
+
+	FSequencerBindingProxy BindingProxy(ObjectBindings[0], Sequencer->GetFocusedMovieSceneSequence());
+
+	TArray<AActor*> SelectedActors;
+	GEditor->GetSelectedActors()->GetSelectedObjects<AActor>(SelectedActors); 
+	ReplaceBindingWithActors(SelectedActors, BindingProxy);
+}
+
+void ULevelSequenceEditorSubsystem::ReplaceBindingWithActors(const TArray<AActor*>& Actors, const FSequencerBindingProxy& ObjectBinding)
+{
+	TSharedPtr<ISequencer> Sequencer = GetActiveSequencer();
+	if (Sequencer == nullptr)
+	{
+		return;
+	}
+
+	FScopedTransaction ReplaceBindingWithActors(LOCTEXT("ReplaceBindingWithActors", "Replace Binding with Actors"));
+
+	FGuid Guid = ObjectBinding.BindingID;
+	TArray<AActor*> ExistingActors;
+	for (TWeakObjectPtr<> Ptr : Sequencer->FindObjectsInCurrentSequence(Guid))
+	{
+		if (AActor* Actor = Cast<AActor>(Ptr.Get()))
+		{
+			if (!Actors.Contains(Actor))
+			{
+				ExistingActors.Add(Actor);
+			}
+		}
+	}
+
+	RemoveActorsFromBinding(ExistingActors, ObjectBinding);
+
+	TArray<AActor*> NewActors;
+	for (AActor* NewActor : Actors)
+	{
+		if (!ExistingActors.Contains(NewActor))
+		{
+			NewActors.Add(NewActor);
+		}
+	}
+
+	AddActorsToBinding(NewActors, ObjectBinding);
+}
+
+void ULevelSequenceEditorSubsystem::RemoveActorsFromBindingInternal()
+{
+	TSharedPtr<ISequencer> Sequencer = GetActiveSequencer();
+	if (Sequencer == nullptr)
+	{
+		return;
+	}
+
+	TArray<FGuid> ObjectBindings;
+	Sequencer->GetSelectedObjects(ObjectBindings);
+	if (ObjectBindings.Num() == 0)
+	{
+		return;
+	}
+
+	FSequencerBindingProxy BindingProxy(ObjectBindings[0], Sequencer->GetFocusedMovieSceneSequence());
+
+	TArray<AActor*> SelectedActors;
+	GEditor->GetSelectedActors()->GetSelectedObjects<AActor>(SelectedActors); 
+	RemoveActorsFromBinding(SelectedActors, BindingProxy);
+}
+
+void ULevelSequenceEditorSubsystem::RemoveActorsFromBinding(const TArray<AActor*>& Actors, const FSequencerBindingProxy& ObjectBinding)
+{
+	if (!Actors.Num())
+	{
+		return;
+	}
+
+	TSharedPtr<ISequencer> Sequencer = GetActiveSequencer();
+	if (Sequencer == nullptr)
+	{
+		return;
+	}
+
+	UMovieSceneSequence* Sequence = Sequencer->GetFocusedMovieSceneSequence();
+	if (!Sequence)
+	{
+		return;
+	}
+
+	UMovieScene* MovieScene = Sequence->GetMovieScene();
+	if (!MovieScene)
+	{
+		return;
+	}
+
+	UClass* ActorClass = nullptr;
+	int32 NumRuntimeObjects = 0;
+
+	FGuid Guid = ObjectBinding.BindingID;
+	for (TWeakObjectPtr<> Ptr : Sequencer->FindObjectsInCurrentSequence(Guid))
+	{
+		if (const AActor* Actor = Cast<AActor>(Ptr.Get()))
+		{
+			ActorClass = Actor->GetClass();
+			++NumRuntimeObjects;
+		}
+	}
+
+	FScopedTransaction RemoveSelectedFromBinding(LOCTEXT("RemoveSelectedFromBinding", "Remove Selected from Binding"));
+
+	TArray<UObject*> ObjectsToRemove;
+	for (AActor* ActorToRemove : Actors)
+	{
+		// Restore state on any components
+		for (UActorComponent* Component : TInlineComponentArray<UActorComponent*>(ActorToRemove))
+		{
+			if (Component)
+			{
+				Sequencer->PreAnimatedState.RestorePreAnimatedState(*Component);
+			}
+		}
+
+		// Restore state on the object itself
+		Sequencer->PreAnimatedState.RestorePreAnimatedState(*ActorToRemove);
+
+		ActorToRemove->Modify();
+
+		ObjectsToRemove.Add(ActorToRemove);
+	}
+	
+	Sequence->Modify();
+	MovieScene->Modify();
+
+	// Unbind objects
+	Sequence->UnbindObjects(Guid, ObjectsToRemove, Sequencer->GetPlaybackContext());
+
+	// Update label
+	if (NumRuntimeObjects - ObjectsToRemove.Num() > 0)
+	{
+		FMovieScenePossessable* Possessable = MovieScene->FindPossessable(Guid);
+		if (Possessable && ActorClass != nullptr)
+		{
+			if (NumRuntimeObjects - ObjectsToRemove.Num() > 1)
+			{
+				FString NewLabel = ActorClass->GetName() + FString::Printf(TEXT(" (%d)"), NumRuntimeObjects - ObjectsToRemove.Num());
+
+				Possessable->SetName(NewLabel);
+			}
+			else if (ObjectsToRemove.Num() > 0 && Actors.Num() > 0)
+			{
+				Possessable->SetName(Actors[0]->GetActorLabel());
+			}
+		}
+	}
+
+	Sequencer->NotifyMovieSceneDataChanged(EMovieSceneDataChangeType::MovieSceneStructureItemsChanged);
+}
+
+void ULevelSequenceEditorSubsystem::RemoveAllBindingsInternal()
+{
+	TSharedPtr<ISequencer> Sequencer = GetActiveSequencer();
+	if (Sequencer == nullptr)
+	{
+		return;
+	}
+
+	TArray<FGuid> ObjectBindings;
+	Sequencer->GetSelectedObjects(ObjectBindings);
+	if (ObjectBindings.Num() == 0)
+	{
+		return;
+	}
+
+	FSequencerBindingProxy BindingProxy(ObjectBindings[0], Sequencer->GetFocusedMovieSceneSequence());
+
+	RemoveAllBindings(BindingProxy);
+}
+
+void ULevelSequenceEditorSubsystem::RemoveAllBindings(const FSequencerBindingProxy& ObjectBinding)
+{
+	TSharedPtr<ISequencer> Sequencer = GetActiveSequencer();
+	if (Sequencer == nullptr)
+	{
+		return;
+	}
+
+	UMovieSceneSequence* Sequence = Sequencer->GetFocusedMovieSceneSequence();
+	if (!Sequence)
+	{
+		return;
+	}
+
+	UMovieScene* MovieScene = Sequence->GetMovieScene();
+	if (!MovieScene)
+	{
+		return;
+	}
+
+	FScopedTransaction RemoveAllBindings(LOCTEXT("RemoveAllBindings", "Remove All Bound Objects"));
+
+	Sequence->Modify();
+	MovieScene->Modify();
+
+	// Unbind objects
+	FGuid Guid = ObjectBinding.BindingID;
+	Sequence->UnbindPossessableObjects(Guid);
+
+	Sequencer->RestorePreAnimatedState();
+
+	Sequencer->NotifyMovieSceneDataChanged(EMovieSceneDataChangeType::MovieSceneStructureItemsChanged);
+}
+
+void ULevelSequenceEditorSubsystem::RemoveInvalidBindingsInternal()
+{
+	TSharedPtr<ISequencer> Sequencer = GetActiveSequencer();
+	if (Sequencer == nullptr)
+	{
+		return;
+	}
+
+	TArray<FGuid> ObjectBindings;
+	Sequencer->GetSelectedObjects(ObjectBindings);
+	if (ObjectBindings.Num() == 0)
+	{
+		return;
+	}
+
+	FSequencerBindingProxy BindingProxy(ObjectBindings[0], Sequencer->GetFocusedMovieSceneSequence());
+
+	RemoveInvalidBindings(BindingProxy);
+}
+
+void ULevelSequenceEditorSubsystem::RemoveInvalidBindings(const FSequencerBindingProxy& ObjectBinding)
+{
+	TSharedPtr<ISequencer> Sequencer = GetActiveSequencer();
+	if (Sequencer == nullptr)
+	{
+		return;
+	}
+
+	UMovieSceneSequence* Sequence = Sequencer->GetFocusedMovieSceneSequence();
+	if (!Sequence)
+	{
+		return;
+	}
+
+	UMovieScene* MovieScene = Sequence->GetMovieScene();
+	if (!MovieScene)
+	{
+		return;
+	}
+	
+	FScopedTransaction RemoveInvalidBindings(LOCTEXT("RemoveMissing", "Remove Missing Objects"));
+
+	Sequence->Modify();
+	MovieScene->Modify();
+
+	// Unbind objects
+	FGuid Guid = ObjectBinding.BindingID;
+	Sequence->UnbindInvalidObjects(Guid, Sequencer->GetPlaybackContext());
+
+	// Update label
+	UClass* ActorClass = nullptr;
+
+	TArray<AActor*> ValidActors;
+	for (TWeakObjectPtr<> Ptr : Sequencer->FindObjectsInCurrentSequence(Guid))
+	{
+		if (AActor* Actor = Cast<AActor>(Ptr.Get()))
+		{
+			ActorClass = Actor->GetClass();
+			ValidActors.Add(Actor);
+		}
+	}
+
+	FMovieScenePossessable* Possessable = MovieScene->FindPossessable(Guid);
+	if (Possessable && ActorClass != nullptr && ValidActors.Num() != 0)
+	{
+		if (ValidActors.Num() > 1)
+		{
+			FString NewLabel = ActorClass->GetName() + FString::Printf(TEXT(" (%d)"), ValidActors.Num());
+
+			Possessable->SetName(NewLabel);
+		}
+		else
+		{
+			Possessable->SetName(ValidActors[0]->GetActorLabel());
+		}
+	}
+
+	Sequencer->RestorePreAnimatedState();
+
+	Sequencer->NotifyMovieSceneDataChanged(EMovieSceneDataChangeType::MovieSceneStructureItemsChanged);
+}
+
 #undef LOCTEXT_NAMESPACE
diff --git a/Engine/Plugins/MovieScene/LevelSequenceEditor/Source/LevelSequenceEditor/Public/LevelSequenceEditorSubsystem.h b/Engine/Plugins/MovieScene/LevelSequenceEditor/Source/LevelSequenceEditor/Public/LevelSequenceEditorSubsystem.h
--- a/Engine/Plugins/MovieScene/LevelSequenceEditor/Source/LevelSequenceEditor/Public/LevelSequenceEditorSubsystem.h
+++ b/Engine/Plugins/MovieScene/LevelSequenceEditor/Source/LevelSequenceEditor/Public/LevelSequenceEditorSubsystem.h
@@ -51,11 +51,36 @@ public:
 	UFUNCTION(BlueprintCallable, Category = "Level Sequence Editor")
 	void FixActorReferences();
 
+	/** Assigns the given actors to the binding */
+	UFUNCTION(BlueprintCallable, Category = "Level Sequence Editor")
+	void AddActorsToBinding(const TArray<AActor*>& Actors, const FSequencerBindingProxy& ObjectBinding);
+
+	/** Replaces the binding with the given actors */
+	UFUNCTION(BlueprintCallable, Category = "Level Sequence Editor")
+	void ReplaceBindingWithActors(const TArray<AActor*>& Actors, const FSequencerBindingProxy& ObjectBinding);
+
+	/** Removes the given actors from the binding */
+	UFUNCTION(BlueprintCallable, Category = "Level Sequence Editor")
+	void RemoveActorsFromBinding(const TArray<AActor*>& Actors, const FSequencerBindingProxy& ObjectBinding);
+
+	/** Remove all bound actors from this track */
+	UFUNCTION(BlueprintCallable, Category = "Level Sequence Editor")
+	void RemoveAllBindings(const FSequencerBindingProxy& ObjectBinding);
+
+	/** Remove missing objects bound to this track */
+	UFUNCTION(BlueprintCallable, Category = "Level Sequence Editor")
+	void RemoveInvalidBindings(const FSequencerBindingProxy& ObjectBinding);
+
 private:
 
 	TSharedPtr<ISequencer> GetActiveSequencer();
 	
 	void BakeTransformInternal();
+	void AddActorsToBindingInternal();
+	void ReplaceBindingWithActorsInternal();
+	void RemoveActorsFromBindingInternal();
+	void RemoveAllBindingsInternal();
+	void RemoveInvalidBindingsInternal();
 
 	FDelegateHandle OnSequencerCreatedHandle;
 
@@ -66,4 +91,6 @@ private:
 
 	TSharedPtr<FExtender> BakeTransformMenuExtender;
 	TSharedPtr<FExtender> FixActorReferencesMenuExtender;
+
+	TSharedPtr<FExtender> AssignActorMenuExtender;
 };
diff --git a/Engine/Source/Editor/Sequencer/Private/DisplayNodes/SequencerObjectBindingNode.cpp b/Engine/Source/Editor/Sequencer/Private/DisplayNodes/SequencerObjectBindingNode.cpp
--- a/Engine/Source/Editor/Sequencer/Private/DisplayNodes/SequencerObjectBindingNode.cpp
+++ b/Engine/Source/Editor/Sequencer/Private/DisplayNodes/SequencerObjectBindingNode.cpp
@@ -295,18 +295,11 @@ void FSequencerObjectBindingNode::BuildContextMenu(FMenuBuilder& MenuBuilder)
 		}
 		else
 		{
-			
-			if (ObjectClass->IsChildOf(AActor::StaticClass()))
-			{
-				FFormatNamedArguments Args;
-
-				MenuBuilder.AddSubMenu(
-					FText::Format(LOCTEXT("AssignActor", "Assign Actor"), Args),
-					FText::Format(LOCTEXT("AssignActorTooltip", "Assign an actor to this track"), Args),
-					FNewMenuDelegate::CreateSP(this, &FSequencerObjectBindingNode::AddAssignActorMenu));
-			}
+			MenuBuilder.BeginSection("Possessable");
 
 			MenuBuilder.AddMenuEntry( FSequencerCommands::Get().ConvertToSpawnable );
+
+			MenuBuilder.EndSection();
 		}
 
 		MenuBuilder.BeginSection("Import/Export", LOCTEXT("ImportExportMenuSectionName", "Import/Export"));
@@ -542,60 +535,6 @@ void FSequencerObjectBindingNode::HandleTemplateActorClassPicked(UClass* ChosenC
 	}
 }
 
-
-void FSequencerObjectBindingNode::AddAssignActorMenu(FMenuBuilder& MenuBuilder)
-{
-	TArray<AActor*> SelectedActors;
-	GEditor->GetSelectedActors()->GetSelectedObjects<AActor>(SelectedActors);
-
-	MenuBuilder.AddMenuEntry(
-		LOCTEXT("AddSelectedToBinding", "Add Selected"),
-		LOCTEXT("AddSelectedToBindingTooltip", "Add selected actors to this track"),
-		FSlateIcon(),
-		FUIAction(
-			FExecuteAction::CreateLambda([=] { GetSequencer().AddActorsToBinding(ObjectBinding, SelectedActors); }),
-			FCanExecuteAction::CreateLambda([=] { return SelectedActors.Num() > 0; })
-			)
-		);
-	MenuBuilder.AddMenuEntry(
-		LOCTEXT("ReplaceBindingWithSelected", "Replace with Selected"),
-		LOCTEXT("ReplaceBindingWithSelectedTooltip", "Replace the object binding with selected actors"),
-		FSlateIcon(),
-		FUIAction(
-			FExecuteAction::CreateLambda([=] { GetSequencer().ReplaceBindingWithActors(ObjectBinding, SelectedActors); }),
-			FCanExecuteAction::CreateLambda([=] { return SelectedActors.Num() > 0; })
-		)
-	);
-	MenuBuilder.AddMenuEntry(
-		LOCTEXT("RemoveSelectedFromBinding", "Remove Selected"),
-		LOCTEXT("RemoveSelectedFromBindingTooltip", "Remove selected actors from this track"),
-		FSlateIcon(),
-		FUIAction(
-			FExecuteAction::CreateLambda([=] { GetSequencer().RemoveActorsFromBinding(ObjectBinding, SelectedActors); }),
-			FCanExecuteAction::CreateLambda([=] { return SelectedActors.Num() > 0; })
-		)
-	);
-	MenuBuilder.AddMenuEntry(
-		LOCTEXT("RemoveAllBindings", "Remove All"),
-		LOCTEXT("RemoveAllBindingsTooltip", "Remove all bound actors from this track"),
-		FSlateIcon(),
-		FUIAction(
-			FExecuteAction::CreateLambda([=] { GetSequencer().RemoveAllBindings(ObjectBinding); })
-		)
-	);
-	MenuBuilder.AddMenuEntry(
-		LOCTEXT("RemoveMissing", "Remove Missing"),
-		LOCTEXT("RemoveMissingooltip", "Remove missing objects bound to this track"),
-		FSlateIcon(),
-		FUIAction(
-			FExecuteAction::CreateLambda([=] { GetSequencer().RemoveInvalidBindings(ObjectBinding); })
-		)
-	);
-
-	GetSequencer().AssignActor(MenuBuilder, ObjectBinding);
-}
-
-
 void FSequencerObjectBindingNode::AddTagMenu(FMenuBuilder& MenuBuilder)
 {
 	MenuBuilder.AddMenuEntry(FSequencerCommands::Get().OpenTaggedBindingManager);
diff --git a/Engine/Source/Editor/Sequencer/Private/DisplayNodes/SequencerObjectBindingNode.h b/Engine/Source/Editor/Sequencer/Private/DisplayNodes/SequencerObjectBindingNode.h
--- a/Engine/Source/Editor/Sequencer/Private/DisplayNodes/SequencerObjectBindingNode.h
+++ b/Engine/Source/Editor/Sequencer/Private/DisplayNodes/SequencerObjectBindingNode.h
@@ -84,7 +84,6 @@ protected:
 
 	void AddSpawnOwnershipMenu(FMenuBuilder& MenuBuilder);
 	void AddSpawnLevelMenu(FMenuBuilder& MenuBuilder);
-	void AddAssignActorMenu(FMenuBuilder& MenuBuilder);
 	void AddTagMenu(FMenuBuilder& MenuBuilder);
 	void AddChangeClassMenu(FMenuBuilder& MenuBuilder);
 
diff --git a/Engine/Source/Editor/Sequencer/Private/Sequencer.cpp b/Engine/Source/Editor/Sequencer/Private/Sequencer.cpp
--- a/Engine/Source/Editor/Sequencer/Private/Sequencer.cpp
+++ b/Engine/Source/Editor/Sequencer/Private/Sequencer.cpp
@@ -7535,7 +7535,7 @@ void FSequencer::AssignActor(FMenuBuilder& MenuBuilder, FGuid InObjectBinding)
 				FOnActorPicked::CreateLambda([=](AActor* Actor){
 					// Create a new binding for this actor
 					FSlateApplication::Get().DismissAllMenus();
-					FSequencerUtilities::DoAssignActor(this, &Actor, 1, InObjectBinding);
+					FSequencerUtilities::DoAssignActor(this, Actor, InObjectBinding);
 				})
 			)
 		];
@@ -7729,70 +7729,6 @@ void FSequencer::RemoveActorsFromBinding(FGuid InObjectBinding, const TArray<AAc
 	NotifyMovieSceneDataChanged(EMovieSceneDataChangeType::MovieSceneStructureItemsChanged);
 }
 
-void FSequencer::RemoveAllBindings(FGuid InObjectBinding)
-{
-	FScopedTransaction RemoveAllBindings(NSLOCTEXT("Sequencer", "RemoveAllBindings", "Remove All Bound Objects"));
-
-	UMovieSceneSequence* OwnerSequence = GetFocusedMovieSceneSequence();
-	UMovieScene* OwnerMovieScene = OwnerSequence->GetMovieScene();
-
-	OwnerSequence->Modify();
-	OwnerMovieScene->Modify();
-
-	// Unbind objects
-	OwnerSequence->UnbindPossessableObjects(InObjectBinding);
-
-	RestorePreAnimatedState();
-
-	NotifyMovieSceneDataChanged(EMovieSceneDataChangeType::MovieSceneStructureItemsChanged);
-}
-
-void FSequencer::RemoveInvalidBindings(FGuid InObjectBinding)
-{
-	FScopedTransaction RemoveInvalidBindings(NSLOCTEXT("Sequencer", "RemoveMissing", "Remove Missing Objects"));
-
-	UMovieSceneSequence* OwnerSequence = GetFocusedMovieSceneSequence();
-	UMovieScene* OwnerMovieScene = OwnerSequence->GetMovieScene();
-
-	OwnerSequence->Modify();
-	OwnerMovieScene->Modify();
-
-	// Unbind objects
-	OwnerSequence->UnbindInvalidObjects(InObjectBinding, GetPlaybackContext());
-
-	// Update label
-	UClass* ActorClass = nullptr;
-
-	TArray<AActor*> ValidActors;
-	for (TWeakObjectPtr<> Ptr : FindObjectsInCurrentSequence(InObjectBinding))
-	{
-		if (AActor* Actor = Cast<AActor>(Ptr.Get()))
-		{
-			ActorClass = Actor->GetClass();
-			ValidActors.Add(Actor);
-		}
-	}
-
-	FMovieScenePossessable* Possessable = OwnerMovieScene->FindPossessable(InObjectBinding);
-	if (Possessable && ActorClass != nullptr && ValidActors.Num() != 0)
-	{
-		if (ValidActors.Num() > 1)
-		{
-			FString NewLabel = ActorClass->GetName() + FString::Printf(TEXT(" (%d)"), ValidActors.Num());
-
-			Possessable->SetName(NewLabel);
-		}
-		else
-		{
-			Possessable->SetName(ValidActors[0]->GetActorLabel());
-		}
-	}
-
-	RestorePreAnimatedState();
-
-	NotifyMovieSceneDataChanged(EMovieSceneDataChangeType::MovieSceneStructureItemsChanged);
-}
-
 void FSequencer::DeleteNode(TSharedRef<FSequencerDisplayNode> NodeToBeDeleted, const bool bKeepState)
 {
 	// If this node is selected, delete all selected nodes
@@ -9031,7 +8967,7 @@ bool FSequencer::PasteObjectBindings(const FString& TextToImport, UMovieSceneFol
 
 					if (!ExistingGuid.IsValid())
 					{
-						FGuid NewGuid = FSequencerUtilities::DoAssignActor(this, &Actor, 1, Possessable->GetGuid());
+						FGuid NewGuid = FSequencerUtilities::DoAssignActor(this, Actor, Possessable->GetGuid());
 
 						// If assigning produces a new guid, update the possessable guids and the bindings pasted data
 						if (NewGuid.IsValid())
diff --git a/Engine/Source/Editor/Sequencer/Private/Sequencer.h b/Engine/Source/Editor/Sequencer/Private/Sequencer.h
--- a/Engine/Source/Editor/Sequencer/Private/Sequencer.h
+++ b/Engine/Source/Editor/Sequencer/Private/Sequencer.h
@@ -528,8 +528,6 @@ public:
 	void AddActorsToBinding(FGuid ObjectBinding, const TArray<AActor*>& InActors);
 	void ReplaceBindingWithActors(FGuid ObjectBinding, const TArray<AActor*>& InActors);
 	void RemoveActorsFromBinding(FGuid ObjectBinding, const TArray<AActor*>& InActors);
-	void RemoveAllBindings(FGuid ObjectBinding);
-	void RemoveInvalidBindings(FGuid ObjectBinding);
 
 	/** Called when a user executes the delete node menu item */
 	void DeleteNode(TSharedRef<FSequencerDisplayNode> NodeToBeDeleted, const bool bKeepState);
diff --git a/Engine/Source/Editor/Sequencer/Private/SequencerUtilities.cpp b/Engine/Source/Editor/Sequencer/Private/SequencerUtilities.cpp
--- a/Engine/Source/Editor/Sequencer/Private/SequencerUtilities.cpp
+++ b/Engine/Source/Editor/Sequencer/Private/SequencerUtilities.cpp
@@ -346,17 +346,8 @@ TArray<FString> FSequencerUtilities::GetAssociatedMapPackages(const ULevelSequen
 	return AssociatedMaps;
 }
 
-FGuid FSequencerUtilities::DoAssignActor(ISequencer * InSequencerPtr, AActor* const* InActors, int32 NumActors, FGuid InObjectBinding)
+FGuid FSequencerUtilities::DoAssignActor(ISequencer * InSequencerPtr, AActor* Actor, FGuid InObjectBinding)
 {
-	if (NumActors <= 0)
-	{
-		return FGuid();
-	}
-
-	//@todo: this code doesn't work with multiple actors, or when the existing binding is bound to multiple actors
-
-	AActor* Actor = InActors[0];
-
 	if (Actor == nullptr)
 	{
 		return FGuid();
diff --git a/Engine/Source/Editor/Sequencer/Public/SequencerUtilities.h b/Engine/Source/Editor/Sequencer/Public/SequencerUtilities.h
--- a/Engine/Source/Editor/Sequencer/Public/SequencerUtilities.h
+++ b/Engine/Source/Editor/Sequencer/Public/SequencerUtilities.h
@@ -32,7 +32,7 @@ struct SEQUENCER_API FSequencerUtilities
 	 */
 	static FName GetUniqueName(FName CandidateName, const TArray<FName>& ExistingNames);
 
-	static FGuid DoAssignActor(ISequencer* InSequencerPtr, AActor* const* InActors, int32 NumActors, FGuid InObjectBinding);
+	static FGuid DoAssignActor(ISequencer* InSequencerPtr, AActor* Actor, FGuid InObjectBinding);
 
 	static void UpdateBindingIDs(ISequencer* InSequencerPtr, UMovieSceneCompiledDataManager* InCompiledDataManagerPtr, FGuid OldGuid, FGuid NewGuid);
 
