Sequencer: Force evaluate when resetting position before evaluating bake transform

diff --git a/Engine/Plugins/MovieScene/LevelSequenceEditor/Source/LevelSequenceEditor/Private/LevelSequenceEditorSubsystem.cpp b/Engine/Plugins/MovieScene/LevelSequenceEditor/Source/LevelSequenceEditor/Private/LevelSequenceEditorSubsystem.cpp
--- a/Engine/Plugins/MovieScene/LevelSequenceEditor/Source/LevelSequenceEditor/Private/LevelSequenceEditorSubsystem.cpp
+++ b/Engine/Plugins/MovieScene/LevelSequenceEditor/Source/LevelSequenceEditor/Private/LevelSequenceEditorSubsystem.cpp
@@ -428,6 +428,7 @@ void ULevelSequenceEditorSubsystem::BakeTransform(const TArray<FSequencerBinding
 
 		// Reset position
 		Sequencer->SetLocalTimeDirectly(ResetTime.Time);
+		Sequencer->ForceEvaluate();
 
 		FVector DefaultLocation = FVector::ZeroVector;
 		FVector DefaultRotation = FVector::ZeroVector;
