LevelSequenceEditor: add command to snap sections to the timeline using their source timecode

This adds a command to the action menu in Sequencer to snap the selected movie
scene sections to the position in the timeline that matches their source timecode.
This can be used to quickly place sections in the timeline in their correct positions
according to their "TimecodeSource" property. Any section with a default value
for "TimecodeSource" (i.e. at timecode "00:00:00:00") will not be modified.

As part of this change, the existing "Sync Sections using Source Timecode" command
was moved out of Sequencer itself and into the LevelSequenceEditorSubsystem,
since it is closely related to the newly added command. Doing so makes both of these
functions accessible to Blueprint and Python. Note though that this means that both
commands are now only available in the Sequencer actions menu, and not in the
"Edit" context menu when right-clicking on a movie scene section.

Finally, this also fixes an issue in the "Sync Sections using Source Timecode" command
where any trimming at the start of the section was not being accounted for, resulting
in incorrect results when syncing trimmed sections.

diff --git a/Engine/Plugins/MovieScene/LevelSequenceEditor/Source/LevelSequenceEditor/Private/LevelSequenceEditorCommands.cpp b/Engine/Plugins/MovieScene/LevelSequenceEditor/Source/LevelSequenceEditor/Private/LevelSequenceEditorCommands.cpp
--- a/Engine/Plugins/MovieScene/LevelSequenceEditor/Source/LevelSequenceEditor/Private/LevelSequenceEditorCommands.cpp
+++ b/Engine/Plugins/MovieScene/LevelSequenceEditor/Source/LevelSequenceEditor/Private/LevelSequenceEditorCommands.cpp
@@ -15,6 +15,9 @@ void FLevelSequenceEditorCommands::RegisterCommands()
 	UI_COMMAND(CreateNewMasterSequenceInLevel, "Add Master Sequence", "Create a new master sequence asset, and place an instance of it in this level", EUserInterfaceActionType::Button, FInputChord());
 	UI_COMMAND(ToggleCinematicViewportCommand, "Cinematic Viewport", "A viewport layout tailored to cinematic preview", EUserInterfaceActionType::RadioButton, FInputChord());
 
+	UI_COMMAND(SnapSectionsToTimelineUsingSourceTimecode, "Snap Sections to Timeline using Source Timecode", "Snap selected sections to the position in the timeline matching their source timecode", EUserInterfaceActionType::Button, FInputChord());
+	UI_COMMAND(SyncSectionsUsingSourceTimecode, "Sync Sections using Source Timecode", "Synchronize sections to the first selected section using source timecode", EUserInterfaceActionType::Button, FInputChord());
+
 	UI_COMMAND(BakeTransform, "Bake Transform", "Bake transform in world space, removing any existing transform and attach tracks", EUserInterfaceActionType::Button, FInputChord());
 	UI_COMMAND(FixActorReferences, "Fix Actor References", "Try to automatically fix up broken actor bindings", EUserInterfaceActionType::Button, FInputChord());
 
diff --git a/Engine/Plugins/MovieScene/LevelSequenceEditor/Source/LevelSequenceEditor/Private/LevelSequenceEditorCommands.h b/Engine/Plugins/MovieScene/LevelSequenceEditor/Source/LevelSequenceEditor/Private/LevelSequenceEditorCommands.h
--- a/Engine/Plugins/MovieScene/LevelSequenceEditor/Source/LevelSequenceEditor/Private/LevelSequenceEditorCommands.h
+++ b/Engine/Plugins/MovieScene/LevelSequenceEditor/Source/LevelSequenceEditor/Private/LevelSequenceEditorCommands.h
@@ -20,6 +20,9 @@ public:
 	TSharedPtr<FUICommandInfo> CreateNewMasterSequenceInLevel;
 	TSharedPtr<FUICommandInfo> ToggleCinematicViewportCommand;
 	
+	TSharedPtr<FUICommandInfo> SnapSectionsToTimelineUsingSourceTimecode;
+	TSharedPtr<FUICommandInfo> SyncSectionsUsingSourceTimecode;
+
 	TSharedPtr<FUICommandInfo> BakeTransform;
 	TSharedPtr<FUICommandInfo> FixActorReferences;
 
diff --git a/Engine/Plugins/MovieScene/LevelSequenceEditor/Source/LevelSequenceEditor/Private/LevelSequenceEditorSubsystem.cpp b/Engine/Plugins/MovieScene/LevelSequenceEditor/Source/LevelSequenceEditor/Private/LevelSequenceEditorSubsystem.cpp
--- a/Engine/Plugins/MovieScene/LevelSequenceEditor/Source/LevelSequenceEditor/Private/LevelSequenceEditorSubsystem.cpp
+++ b/Engine/Plugins/MovieScene/LevelSequenceEditor/Source/LevelSequenceEditor/Private/LevelSequenceEditorSubsystem.cpp
@@ -8,6 +8,7 @@
 #include "LevelSequenceEditorCommands.h"
 #include "MovieScene.h"
 #include "MovieSceneCommonHelpers.h"
+#include "MovieSceneSection.h"
 #include "MovieSceneSequence.h"
 #include "MovieSceneTimeHelpers.h"
 #include "SequencerSettings.h"
@@ -60,9 +61,31 @@ void ULevelSequenceEditorSubsystem::Initialize(FSubsystemCollectionBase& Collect
 		return SelectedActors.Num() > 0;
 	};
 
+	auto AreMovieSceneSectionsSelected = [this](const int32 MinSections = 1) {
+		const TSharedPtr<ISequencer> Sequencer = GetActiveSequencer();
+		if (!Sequencer)
+		{
+			return false;
+		}
+
+		TArray<UMovieSceneSection*> SelectedSections;
+		Sequencer->GetSelectedSections(SelectedSections);
+		return (SelectedSections.Num() >= MinSections);
+	};
+
 	/* Commands for this subsystem */
 	CommandList = MakeShareable(new FUICommandList);
 
+	CommandList->MapAction(FLevelSequenceEditorCommands::Get().SnapSectionsToTimelineUsingSourceTimecode,
+		FExecuteAction::CreateUObject(this, &ULevelSequenceEditorSubsystem::SnapSectionsToTimelineUsingSourceTimecodeInternal),
+		FCanExecuteAction::CreateLambda(AreMovieSceneSectionsSelected)
+	);
+
+	CommandList->MapAction(FLevelSequenceEditorCommands::Get().SyncSectionsUsingSourceTimecode,
+		FExecuteAction::CreateUObject(this, &ULevelSequenceEditorSubsystem::SyncSectionsUsingSourceTimecodeInternal),
+		FCanExecuteAction::CreateLambda(AreMovieSceneSectionsSelected, 2)
+	);
+
 	CommandList->MapAction(FLevelSequenceEditorCommands::Get().BakeTransform,
 		FExecuteAction::CreateUObject(this, &ULevelSequenceEditorSubsystem::BakeTransformInternal)
 	);
@@ -90,12 +113,14 @@ void ULevelSequenceEditorSubsystem::Initialize(FSubsystemCollectionBase& Collect
 		FExecuteAction::CreateUObject(this, &ULevelSequenceEditorSubsystem::RemoveInvalidBindingsInternal));
 
 	/* Menu extenders */
-	BakeTransformMenuExtender = MakeShareable(new FExtender);
-	BakeTransformMenuExtender->AddMenuExtension("Transform", EExtensionHook::After, CommandList, FMenuExtensionDelegate::CreateStatic([](FMenuBuilder& MenuBuilder) {
+	TransformMenuExtender = MakeShareable(new FExtender);
+	TransformMenuExtender->AddMenuExtension("Transform", EExtensionHook::After, CommandList, FMenuExtensionDelegate::CreateStatic([](FMenuBuilder& MenuBuilder) {
+		MenuBuilder.AddMenuEntry(FLevelSequenceEditorCommands::Get().SnapSectionsToTimelineUsingSourceTimecode);
+		MenuBuilder.AddMenuEntry(FLevelSequenceEditorCommands::Get().SyncSectionsUsingSourceTimecode);
 		MenuBuilder.AddMenuEntry(FLevelSequenceEditorCommands::Get().BakeTransform);
 		}));
 
-	SequencerModule.GetActionsMenuExtensibilityManager()->AddExtender(BakeTransformMenuExtender);
+	SequencerModule.GetActionsMenuExtensibilityManager()->AddExtender(TransformMenuExtender);
 
 	FixActorReferencesMenuExtender = MakeShareable(new FExtender);
 	FixActorReferencesMenuExtender->AddMenuExtension("Bindings", EExtensionHook::First, CommandList, FMenuExtensionDelegate::CreateStatic([](FMenuBuilder& MenuBuilder) {
@@ -148,6 +173,195 @@ TSharedPtr<ISequencer> ULevelSequenceEditorSubsystem::GetActiveSequencer()
 	return nullptr;
 }
 
+void ULevelSequenceEditorSubsystem::SnapSectionsToTimelineUsingSourceTimecodeInternal()
+{
+	const TSharedPtr<ISequencer> Sequencer = GetActiveSequencer();
+	if (!Sequencer)
+	{
+		return;
+	}
+
+	const UMovieScene* FocusedMovieScene = Sequencer->GetFocusedMovieSceneSequence()->GetMovieScene();
+	if (!FocusedMovieScene)
+	{
+		return;
+	}
+
+	if (FocusedMovieScene->IsReadOnly())
+	{
+		FSequencerUtilities::ShowReadOnlyError();
+		return;
+	}
+
+	TArray<UMovieSceneSection*> Sections;
+	Sequencer->GetSelectedSections(Sections);
+	if (Sections.IsEmpty())
+	{
+		return;
+	}
+
+	SnapSectionsToTimelineUsingSourceTimecode(Sections);
+}
+
+void ULevelSequenceEditorSubsystem::SnapSectionsToTimelineUsingSourceTimecode(const TArray<UMovieSceneSection*>& Sections)
+{
+	const TSharedPtr<ISequencer> Sequencer = GetActiveSequencer();
+	if (!Sequencer)
+	{
+		return;
+	}
+
+	const UMovieScene* FocusedMovieScene = Sequencer->GetFocusedMovieSceneSequence()->GetMovieScene();
+	if (!FocusedMovieScene)
+	{
+		return;
+	}
+
+	if (FocusedMovieScene->IsReadOnly())
+	{
+		FSequencerUtilities::ShowReadOnlyError();
+		return;
+	}
+
+	FScopedTransaction SnapSectionsToTimelineUsingSourceTimecodeTransaction(LOCTEXT("SnapSectionsToTimelineUsingSourceTimecode_Transaction", "Snap Sections to Timeline using Source Timecode"));
+	bool bAnythingChanged = false;
+
+	const FFrameRate TickResolution = Sequencer->GetFocusedTickResolution();
+
+	for (UMovieSceneSection* Section : Sections)
+	{
+		if (!Section || !Section->HasStartFrame())
+		{
+			continue;
+		}
+
+		const FTimecode SectionSourceTimecode = Section->TimecodeSource.Timecode;
+		if (SectionSourceTimecode == FTimecode())
+		{
+			// Do not move sections with default values for source timecode.
+			continue;
+		}
+
+		const FFrameNumber SectionSourceStartFrameNumber = SectionSourceTimecode.ToFrameNumber(TickResolution);
+
+		// Account for any trimming at the start of the section when computing the
+		// target frame number to move this section to.
+		const FFrameNumber SectionOffsetFrames = Section->GetOffsetTime().Get(FFrameTime()).FloorToFrame();
+		const FFrameNumber TargetFrameNumber = SectionSourceStartFrameNumber + SectionOffsetFrames;
+
+		const FFrameNumber SectionCurrentStartFrameNumber = Section->GetInclusiveStartFrame();
+
+		const FFrameNumber Delta = -(SectionCurrentStartFrameNumber - TargetFrameNumber);
+
+		Section->MoveSection(Delta);
+
+		bAnythingChanged |= (Delta.Value != 0);
+	}
+
+	if (bAnythingChanged)
+	{
+		Sequencer->NotifyMovieSceneDataChanged(EMovieSceneDataChangeType::TrackValueChanged);
+	}
+}
+
+void ULevelSequenceEditorSubsystem::SyncSectionsUsingSourceTimecodeInternal()
+{
+	const TSharedPtr<ISequencer> Sequencer = GetActiveSequencer();
+	if (!Sequencer)
+	{
+		return;
+	}
+
+	const UMovieScene* FocusedMovieScene = Sequencer->GetFocusedMovieSceneSequence()->GetMovieScene();
+	if (!FocusedMovieScene)
+	{
+		return;
+	}
+
+	if (FocusedMovieScene->IsReadOnly())
+	{
+		FSequencerUtilities::ShowReadOnlyError();
+		return;
+	}
+
+	TArray<UMovieSceneSection*> Sections;
+	Sequencer->GetSelectedSections(Sections);
+	if (Sections.Num() < 2)
+	{
+		return;
+	}
+
+	SyncSectionsUsingSourceTimecode(Sections);
+}
+
+void ULevelSequenceEditorSubsystem::SyncSectionsUsingSourceTimecode(const TArray<UMovieSceneSection*>& Sections)
+{
+	const TSharedPtr<ISequencer> Sequencer = GetActiveSequencer();
+	if (!Sequencer)
+	{
+		return;
+	}
+
+	const UMovieScene* FocusedMovieScene = Sequencer->GetFocusedMovieSceneSequence()->GetMovieScene();
+	if (!FocusedMovieScene)
+	{
+		return;
+	}
+
+	if (FocusedMovieScene->IsReadOnly())
+	{
+		FSequencerUtilities::ShowReadOnlyError();
+		return;
+	}
+
+	// Pull out all of the valid sections that have a start frame and verify
+	// we have at least two sections to sync.
+	TArray<UMovieSceneSection*> SectionsToSync;
+	for (UMovieSceneSection* Section : Sections)
+	{
+		if (Section && Section->HasStartFrame())
+		{
+			SectionsToSync.Add(Section);
+		}
+	}
+
+	if (SectionsToSync.Num() < 2)
+	{
+		return;
+	}
+
+	FScopedTransaction SyncSectionsUsingSourceTimecodeTransaction(LOCTEXT("SyncSectionsUsingSourceTimecode_Transaction", "Sync Sections Using Source Timecode"));
+	bool bAnythingChanged = false;
+
+	const FFrameRate TickResolution = Sequencer->GetFocusedTickResolution();
+
+	const UMovieSceneSection* FirstSection = SectionsToSync[0];
+	const FFrameNumber FirstSectionSourceTimecode = FirstSection->TimecodeSource.Timecode.ToFrameNumber(TickResolution);
+	const FFrameNumber FirstSectionCurrentStartFrame = FirstSection->GetInclusiveStartFrame();
+	const FFrameNumber FirstSectionOffsetFrames = FirstSection->GetOffsetTime().Get(FFrameTime()).FloorToFrame();
+	SectionsToSync.RemoveAt(0);
+
+	for (UMovieSceneSection* Section : SectionsToSync)
+	{
+		const FFrameNumber SectionSourceTimecode = Section->TimecodeSource.Timecode.ToFrameNumber(TickResolution);
+		const FFrameNumber SectionCurrentStartFrame = Section->GetInclusiveStartFrame();
+		const FFrameNumber SectionOffsetFrames = Section->GetOffsetTime().Get(FFrameTime()).FloorToFrame();
+
+		const FFrameNumber TimecodeDelta = SectionSourceTimecode - FirstSectionSourceTimecode;
+		const FFrameNumber CurrentDelta = (SectionCurrentStartFrame - SectionOffsetFrames) - (FirstSectionCurrentStartFrame - FirstSectionOffsetFrames);
+		const FFrameNumber Delta = -CurrentDelta + TimecodeDelta;
+
+		Section->MoveSection(Delta);
+
+		bAnythingChanged |= (Delta.Value != 0);
+	}
+
+	if (bAnythingChanged)
+	{
+		Sequencer->NotifyMovieSceneDataChanged(EMovieSceneDataChangeType::TrackValueChanged);
+	}
+}
+
 void ULevelSequenceEditorSubsystem::BakeTransformInternal()
 {
 	TSharedPtr<ISequencer> Sequencer = GetActiveSequencer();
diff --git a/Engine/Plugins/MovieScene/LevelSequenceEditor/Source/LevelSequenceEditor/Public/LevelSequenceEditorSubsystem.h b/Engine/Plugins/MovieScene/LevelSequenceEditor/Source/LevelSequenceEditor/Public/LevelSequenceEditorSubsystem.h
--- a/Engine/Plugins/MovieScene/LevelSequenceEditor/Source/LevelSequenceEditor/Public/LevelSequenceEditorSubsystem.h
+++ b/Engine/Plugins/MovieScene/LevelSequenceEditor/Source/LevelSequenceEditor/Public/LevelSequenceEditorSubsystem.h
@@ -16,6 +16,7 @@ DECLARE_LOG_CATEGORY_EXTERN(LogLevelSequenceEditor, Log, All);
 
 class FExtender;
 class UMovieSceneCompiledDataManager;
+class UMovieSceneSection;
 
 USTRUCT(BlueprintType)
 struct FMovieSceneScriptingParams
@@ -43,6 +44,14 @@ public:
 
 	void OnSequencerCreated(TSharedRef<ISequencer> InSequencer);
 	
+	/** Snap sections to timeline using source timecode */
+	UFUNCTION(BlueprintCallable, Category = "Level Sequence Editor")
+	void SnapSectionsToTimelineUsingSourceTimecode(const TArray<UMovieSceneSection*>& Sections);
+
+	/** Sync section using source timecode */
+	UFUNCTION(BlueprintCallable, Category = "Level Sequence Editor")
+	void SyncSectionsUsingSourceTimecode(const TArray<UMovieSceneSection*>& Sections);
+
 	/** Bake transform */
 	UFUNCTION(BlueprintCallable, Category = "Level Sequence Editor")
 	void BakeTransform(const TArray<FSequencerBindingProxy>& ObjectBindings, const FFrameTime& BakeInTime, const FFrameTime& BakeOutTime, const FFrameTime& BakeInterval, const FMovieSceneScriptingParams& Params = FMovieSceneScriptingParams());
@@ -75,6 +84,8 @@ private:
 
 	TSharedPtr<ISequencer> GetActiveSequencer();
 	
+	void SnapSectionsToTimelineUsingSourceTimecodeInternal();
+	void SyncSectionsUsingSourceTimecodeInternal();
 	void BakeTransformInternal();
 	void AddActorsToBindingInternal();
 	void ReplaceBindingWithActorsInternal();
@@ -89,7 +100,7 @@ private:
 
 	TSharedPtr<FUICommandList> CommandList;
 
-	TSharedPtr<FExtender> BakeTransformMenuExtender;
+	TSharedPtr<FExtender> TransformMenuExtender;
 	TSharedPtr<FExtender> FixActorReferencesMenuExtender;
 
 	TSharedPtr<FExtender> AssignActorMenuExtender;
diff --git a/Engine/Source/Editor/Sequencer/Private/SSequencer.cpp b/Engine/Source/Editor/Sequencer/Private/SSequencer.cpp
--- a/Engine/Source/Editor/Sequencer/Private/SSequencer.cpp
+++ b/Engine/Source/Editor/Sequencer/Private/SSequencer.cpp
@@ -2069,11 +2069,6 @@ TSharedRef<SWidget> SSequencer::MakeActionsMenu()
 	MenuBuilder.AddMenuEntry(FSequencerCommands::Get().TrimOrExtendSectionLeft);
 	MenuBuilder.AddMenuEntry(FSequencerCommands::Get().TrimOrExtendSectionRight);
 
-	if (SequencerPtr.Pin()->IsLevelEditorSequencer())
-	{
-		MenuBuilder.AddMenuEntry(FSequencerCommands::Get().SyncSectionsUsingSourceTimecode);
-	}
-
 	MenuBuilder.EndSection();
 
 	// selection range actions
diff --git a/Engine/Source/Editor/Sequencer/Private/Sequencer.cpp b/Engine/Source/Editor/Sequencer/Private/Sequencer.cpp
--- a/Engine/Source/Editor/Sequencer/Private/Sequencer.cpp
+++ b/Engine/Source/Editor/Sequencer/Private/Sequencer.cpp
@@ -2230,54 +2230,6 @@ bool FSequencer::IsPilotCamera() const
 	return false;
 }
 
-void FSequencer::SyncSectionsUsingSourceTimecode()
-{
-	FScopedTransaction SyncSectionsUsingSourceTimecodeTransaction( LOCTEXT("SyncSectionsUsingSourceTimecode_Transaction", "Sync Sections Using Source Timecode") );
-	bool bAnythingChanged = false;
-
-	TArray<UMovieSceneSection*> Sections;
-	for (auto Section : GetSelection().GetSelectedSections())
-	{
-		if (Section.IsValid() && Section->HasStartFrame())
-		{
-			Sections.Add(Section.Get());
-		}
-	}
-
-	if (Sections.Num() < 2) 
-	{
-		return;
-	}
-
-	const UMovieSceneSection* FirstSection = Sections[0];
-	FFrameNumber FirstSectionSourceTimecode = FirstSection->TimecodeSource.Timecode.ToFrameNumber(GetFocusedTickResolution());
-	FFrameNumber FirstSectionCurrentStartFrame = FirstSection->GetInclusiveStartFrame();
-	Sections.RemoveAt(0);
-
-	for (auto Section : Sections)
-	{
-		if (Section->HasStartFrame())
-		{
-			FFrameNumber SectionSourceTimecode = Section->TimecodeSource.Timecode.ToFrameNumber(GetFocusedTickResolution());
-			FFrameNumber SectionCurrentStartFrame = Section->GetInclusiveStartFrame();
-
-			FFrameNumber TimecodeDelta = SectionSourceTimecode - FirstSectionSourceTimecode;
-			FFrameNumber CurrentDelta = SectionCurrentStartFrame - FirstSectionCurrentStartFrame;
-			FFrameNumber Delta = -CurrentDelta + TimecodeDelta;
-
-			Section->MoveSection(Delta);
-
-			bAnythingChanged = bAnythingChanged || (Delta.Value != 0);
-		}
-	}
-
-	if (bAnythingChanged)
-	{
-		NotifyMovieSceneDataChanged( EMovieSceneDataChangeType::TrackValueChanged );
-	}
-}
-
-
 void FSequencer::OnActorsDropped( const TArray<TWeakObjectPtr<AActor> >& Actors )
 {
 	AddActors(Actors);
@@ -12971,11 +12923,6 @@ void FSequencer::BindCommands()
 		FIsActionButtonVisible::CreateLambda([this] { return ExactCast<ULevelSequence>(GetFocusedMovieSceneSequence()) != nullptr && IVREditorModule::Get().IsVREditorModeActive() == false; }) //@todo VREditor: Creating a camera while in VR mode disrupts the hmd. This is a temporary fix by hiding the button when in VR mode.
 	);
 
-	SequencerCommandBindings->MapAction(
-		Commands.SyncSectionsUsingSourceTimecode,
-		FExecuteAction::CreateSP( this, &FSequencer::SyncSectionsUsingSourceTimecode ),
-		FCanExecuteAction::CreateLambda( [this]{ return (GetSelection().GetSelectedSections().Num() > 1); } ) );
-
 	SequencerCommandBindings->MapAction(
 		Commands.FixPossessableObjectClass,
 		FExecuteAction::CreateSP( this, &FSequencer::FixPossessableObjectClass ),
diff --git a/Engine/Source/Editor/Sequencer/Private/Sequencer.h b/Engine/Source/Editor/Sequencer/Private/Sequencer.h
--- a/Engine/Source/Editor/Sequencer/Private/Sequencer.h
+++ b/Engine/Source/Editor/Sequencer/Private/Sequencer.h
@@ -264,9 +264,6 @@ public:
 	/** Shrink time*/
 	void ShrinkTime(FFrameTime InDeltaTime);
 
-	/** Sync using source timecode */
-	void SyncSectionsUsingSourceTimecode();
-
 	/**
 	 * @return Movie scene tools used by the sequencer
 	 */
diff --git a/Engine/Source/Editor/Sequencer/Private/SequencerCommands.cpp b/Engine/Source/Editor/Sequencer/Private/SequencerCommands.cpp
--- a/Engine/Source/Editor/Sequencer/Private/SequencerCommands.cpp
+++ b/Engine/Source/Editor/Sequencer/Private/SequencerCommands.cpp
@@ -111,7 +111,6 @@ void FSequencerCommands::RegisterCommands()
 	UI_COMMAND( OpenDirectorBlueprint, "Open Director Blueprint", "Opens the director blueprint for this sequence.", EUserInterfaceActionType::Button, FInputChord() );
 	UI_COMMAND( OpenTaggedBindingManager, "Open Binding Tag Manager", "Specifies options for tagging bindings within this sequence for external systems to reference as a persistent name.", EUserInterfaceActionType::Button, FInputChord() );
 	UI_COMMAND( OpenNodeGroupsManager, "Open Sequencer Group Manager", "Manage groups within this sequence.", EUserInterfaceActionType::Button, FInputChord());
-	UI_COMMAND( SyncSectionsUsingSourceTimecode, "Sync Sections using Source Timecode", "Synchronize sections to the first selected section using source timecode", EUserInterfaceActionType::Button, FInputChord() );
 	UI_COMMAND( ToggleShowRangeSlider, "Range Slider", "Enables and disables showing the time range slider", EUserInterfaceActionType::ToggleButton, FInputChord() );
 	UI_COMMAND( ToggleIsSnapEnabled, "Enable Snapping", "Enables and disables snapping", EUserInterfaceActionType::ToggleButton, FInputChord() );
 
diff --git a/Engine/Source/Editor/Sequencer/Private/SequencerContextMenus.cpp b/Engine/Source/Editor/Sequencer/Private/SequencerContextMenus.cpp
--- a/Engine/Source/Editor/Sequencer/Private/SequencerContextMenus.cpp
+++ b/Engine/Source/Editor/Sequencer/Private/SequencerContextMenus.cpp
@@ -480,15 +480,6 @@ void FSectionContextMenu::AddEditMenu(FMenuBuilder& MenuBuilder)
 			FCanExecuteAction::CreateLambda([=]{ return Shared->CanAutoSize(); }))
 	);
 
-	MenuBuilder.AddMenuEntry(
-		LOCTEXT("SyncSectionsUsingSourceTimecode", "Synchronize using Source Timecode"),
-		LOCTEXT("SyncSectionsUsingSourceTimecodeTooltip", "Sync selected sections using the source timecode.  The first selected section will be unchanged and subsequent sections will be adjusted according to their source timecode as relative to the first section's."),
-		FSlateIcon(),
-		FUIAction(
-			FExecuteAction::CreateLambda([=] { return Shared->Sequencer->SyncSectionsUsingSourceTimecode(); }),
-			FCanExecuteAction::CreateLambda([=]{ return (Shared->Sequencer->GetSelection().GetSelectedSections().Num() > 1); }))
-	);
-
 	MenuBuilder.BeginSection("SequencerInterpolation", LOCTEXT("KeyInterpolationMenu", "Key Interpolation"));
 	
 	MenuBuilder.AddMenuEntry(
diff --git a/Engine/Source/Editor/Sequencer/Public/SequencerCommands.h b/Engine/Source/Editor/Sequencer/Public/SequencerCommands.h
--- a/Engine/Source/Editor/Sequencer/Public/SequencerCommands.h
+++ b/Engine/Source/Editor/Sequencer/Public/SequencerCommands.h
@@ -286,9 +286,6 @@ public:
 	/** Remove selected nodes from folder. */
 	TSharedPtr< FUICommandInfo > RemoveFromFolder;
 
-	/** Sync sections using source timecode. */
-	TSharedPtr< FUICommandInfo > SyncSectionsUsingSourceTimecode;
-
 	/** Turns the range slider on and off. */
 	TSharedPtr< FUICommandInfo > ToggleShowRangeSlider;
 
