Sequencer: Add rebind to component

diff --git a/Engine/Plugins/MovieScene/LevelSequenceEditor/Source/LevelSequenceEditor/Private/LevelSequenceEditorSubsystem.cpp b/Engine/Plugins/MovieScene/LevelSequenceEditor/Source/LevelSequenceEditor/Private/LevelSequenceEditorSubsystem.cpp
--- a/Engine/Plugins/MovieScene/LevelSequenceEditor/Source/LevelSequenceEditor/Private/LevelSequenceEditorSubsystem.cpp
+++ b/Engine/Plugins/MovieScene/LevelSequenceEditor/Source/LevelSequenceEditor/Private/LevelSequenceEditorSubsystem.cpp
@@ -139,6 +139,22 @@ void ULevelSequenceEditorSubsystem::Initialize(FSubsystemCollectionBase& Collect
 		}));
 
 	SequencerModule.GetObjectBindingContextMenuExtensibilityManager()->AddExtender(AssignActorMenuExtender);
+
+	RebindComponentMenuExtender = MakeShareable(new FExtender);
+	RebindComponentMenuExtender->AddMenuExtension("Possessable", EExtensionHook::First, CommandList, FMenuExtensionDelegate::CreateLambda([this](FMenuBuilder& MenuBuilder) {
+		TArray<FName> ComponentNames;
+		GetRebindComponentNames(ComponentNames);
+		if (ComponentNames.Num() > 0)
+		{
+			FFormatNamedArguments Args;
+			MenuBuilder.AddSubMenu(
+				FText::Format(LOCTEXT("RebindComponent", "Rebind Component"), Args),
+				FText::Format(LOCTEXT("RebindComponentTooltip", "Rebind component by moving the tracks from one component to another component."), Args),
+				FNewMenuDelegate::CreateLambda([this](FMenuBuilder& SubMenuBuilder) { RebindComponentMenu(SubMenuBuilder); } ));
+		}
+	}));
+
+	SequencerModule.GetObjectBindingContextMenuExtensibilityManager()->AddExtender(RebindComponentMenuExtender);
 }
 
 void ULevelSequenceEditorSubsystem::Deinitialize()
@@ -1248,4 +1264,169 @@ void ULevelSequenceEditorSubsystem::RemoveInvalidBindings(const FSequencerBindin
 	Sequencer->NotifyMovieSceneDataChanged(EMovieSceneDataChangeType::MovieSceneStructureItemsChanged);
 }
 
+void ULevelSequenceEditorSubsystem::GetRebindComponentNames(TArray<FName>& OutComponentNames)
+{
+	OutComponentNames.Empty();
+
+	TSharedPtr<ISequencer> Sequencer = GetActiveSequencer();
+	if (Sequencer == nullptr)
+	{
+		return;
+	}
+
+	UMovieSceneSequence* Sequence = Sequencer->GetFocusedMovieSceneSequence();
+	if (!Sequence)
+	{
+		return;
+	}
+
+	UMovieScene* MovieScene = Sequence->GetMovieScene();
+	if (!MovieScene)
+	{
+		return;
+	}
+
+	TArray<FGuid> ObjectBindings;
+	Sequencer->GetSelectedObjects(ObjectBindings);
+	if (ObjectBindings.Num() == 0)
+	{
+		return;
+	}
+
+	FGuid ComponentGuid = ObjectBindings[0];
+
+	FMovieScenePossessable* ComponentPossessable = MovieScene->FindPossessable(ComponentGuid);
+
+	FGuid ActorParentGuid = ComponentPossessable ? ComponentPossessable->GetParent() : FGuid();
+
+	TArrayView<TWeakObjectPtr<>> ObjectsInCurrentSequence = Sequencer->FindObjectsInCurrentSequence(ActorParentGuid);
+
+	const AActor* Actor = nullptr;
+	for (TWeakObjectPtr<> Ptr : ObjectsInCurrentSequence)
+	{
+		Actor = Cast<AActor>(Ptr.Get());
+		if (Actor)
+		{
+			break;
+		}
+	}
+
+	if (!Actor)
+	{
+		return;
+	}
+		
+	for (UActorComponent* Component : Actor->GetComponents())
+	{
+		if (Component && Component->GetName() != ComponentPossessable->GetName())
+		{
+			OutComponentNames.Add(Component->GetFName());
+		}
+	}
+	OutComponentNames.Sort(FNameFastLess());
+}
+
+void ULevelSequenceEditorSubsystem::RebindComponentMenu(FMenuBuilder& MenuBuilder)
+{
+	TArray<FName> ComponentNames;
+	GetRebindComponentNames(ComponentNames);
+
+	for (const FName& ComponentName : ComponentNames)
+	{
+		FText RebindComponentLabel = FText::FromName(ComponentName);
+		MenuBuilder.AddMenuEntry(
+			RebindComponentLabel, 
+			FText(), 
+			FSlateIcon(), 
+			FUIAction(FExecuteAction::CreateLambda([this, ComponentName]() { RebindComponentInternal(ComponentName); } ) ) );
+	}
+}
+
+void ULevelSequenceEditorSubsystem::RebindComponentInternal(const FName& ComponentName)
+{
+	TSharedPtr<ISequencer> Sequencer = GetActiveSequencer();
+	if (Sequencer == nullptr)
+	{
+		return;
+	}
+
+	TArray<FGuid> ObjectBindings;
+	Sequencer->GetSelectedObjects(ObjectBindings);
+	if (ObjectBindings.Num() == 0)
+	{
+		return;
+	}
+
+	TArray<FSequencerBindingProxy> BindingProxies;
+	for (const FGuid& ObjectBinding : ObjectBindings)
+	{
+		FSequencerBindingProxy BindingProxy(ObjectBinding, Sequencer->GetFocusedMovieSceneSequence());
+		BindingProxies.Add(BindingProxy);
+	}
+
+	RebindComponent(BindingProxies, ComponentName);
+}
+
+void ULevelSequenceEditorSubsystem::RebindComponent(const TArray<FSequencerBindingProxy>& PossessableBindings, const FName& ComponentName)
+{
+	TSharedPtr<ISequencer> Sequencer = GetActiveSequencer();
+	if (Sequencer == nullptr)
+	{
+		return;
+	}
+
+	UMovieSceneSequence* Sequence = Sequencer->GetFocusedMovieSceneSequence();
+	if (!Sequence)
+	{
+		return;
+	}
+
+	UMovieScene* MovieScene = Sequence->GetMovieScene();
+	if (!MovieScene)
+	{
+		return;
+	}
+
+	FScopedTransaction RebindComponent(LOCTEXT("RebindComponent", "Rebind Component"));
+
+	Sequence->Modify();
+	MovieScene->Modify();
+
+	bool bAnythingChanged = false;
+	for (const FSequencerBindingProxy& PossessableBinding : PossessableBindings)
+	{
+		FMovieScenePossessable* ComponentPossessable = MovieScene->FindPossessable(PossessableBinding.BindingID);
+
+		FGuid ActorParentGuid = ComponentPossessable ? ComponentPossessable->GetParent() : FGuid();
+
+		TArrayView<TWeakObjectPtr<>> ObjectsInCurrentSequence = Sequencer->FindObjectsInCurrentSequence(ActorParentGuid);
+
+		for (TWeakObjectPtr<> Ptr : ObjectsInCurrentSequence)
+		{
+			if (const AActor* Actor = Cast<AActor>(Ptr.Get()))
+			{
+				for (UActorComponent* Component : Actor->GetComponents())
+				{
+					if (Component->GetFName() == ComponentName)
+					{
+						FGuid ComponentBinding = Sequence->CreatePossessable(Component);
+						
+						if (PossessableBinding.BindingID.IsValid() && ComponentBinding.IsValid())
+						{
+							MovieScene->MoveBindingContents(PossessableBinding.BindingID, ComponentBinding);
+
+							bAnythingChanged = true;
+						}
+					}
+				}
+			}
+		}
+	}
+
+	if (bAnythingChanged)
+	{
+		Sequencer->NotifyMovieSceneDataChanged(EMovieSceneDataChangeType::MovieSceneStructureItemsChanged);
+	}
+}
+
 #undef LOCTEXT_NAMESPACE
diff --git a/Engine/Plugins/MovieScene/LevelSequenceEditor/Source/LevelSequenceEditor/Public/LevelSequenceEditorSubsystem.h b/Engine/Plugins/MovieScene/LevelSequenceEditor/Source/LevelSequenceEditor/Public/LevelSequenceEditorSubsystem.h
--- a/Engine/Plugins/MovieScene/LevelSequenceEditor/Source/LevelSequenceEditor/Public/LevelSequenceEditorSubsystem.h
+++ b/Engine/Plugins/MovieScene/LevelSequenceEditor/Source/LevelSequenceEditor/Public/LevelSequenceEditorSubsystem.h
@@ -80,6 +80,10 @@ public:
 	UFUNCTION(BlueprintCallable, Category = "Level Sequence Editor")
 	void RemoveInvalidBindings(const FSequencerBindingProxy& ObjectBinding);
 
+	/** Rebind the component binding to the requested component */
+	UFUNCTION(BlueprintCallable, Category = "Level Sequence Editor")
+	void RebindComponent(const TArray<FSequencerBindingProxy>& ComponentBindings, const FName& ComponentName);
+
 private:
 
 	TSharedPtr<ISequencer> GetActiveSequencer();
@@ -92,6 +96,10 @@ private:
 	void RemoveActorsFromBindingInternal();
 	void RemoveAllBindingsInternal();
 	void RemoveInvalidBindingsInternal();
+	void RebindComponentInternal(const FName& ComponentName);
+
+	void GetRebindComponentNames(TArray<FName>& OutComponentNames);
+	void RebindComponentMenu(FMenuBuilder& MenuBuilder);
 
 	FDelegateHandle OnSequencerCreatedHandle;
 
@@ -104,4 +112,5 @@ private:
 	TSharedPtr<FExtender> FixActorReferencesMenuExtender;
 
 	TSharedPtr<FExtender> AssignActorMenuExtender;
+	TSharedPtr<FExtender> RebindComponentMenuExtender;
 };
