Sequencer: Fix static analysis

diff --git a/Engine/Plugins/MovieScene/LevelSequenceEditor/Source/LevelSequenceEditor/Private/LevelSequenceEditorSubsystem.cpp b/Engine/Plugins/MovieScene/LevelSequenceEditor/Source/LevelSequenceEditor/Private/LevelSequenceEditorSubsystem.cpp
--- a/Engine/Plugins/MovieScene/LevelSequenceEditor/Source/LevelSequenceEditor/Private/LevelSequenceEditorSubsystem.cpp
+++ b/Engine/Plugins/MovieScene/LevelSequenceEditor/Source/LevelSequenceEditor/Private/LevelSequenceEditorSubsystem.cpp
@@ -1318,7 +1318,7 @@ void ULevelSequenceEditorSubsystem::GetRebindComponentNames(TArray<FName>& OutCo
 		
 	for (UActorComponent* Component : Actor->GetComponents())
 	{
-		if (Component && Component->GetName() != ComponentPossessable->GetName())
+		if (Component && ComponentPossessable && Component->GetName() != ComponentPossessable->GetName())
 		{
 			OutComponentNames.Add(Component->GetFName());
 		}
