Sequencer: Follow up to 18551247 - Assign Actor missing scene outliner chooser

diff --git a/Engine/Plugins/MovieScene/LevelSequenceEditor/Source/LevelSequenceEditor/Private/LevelSequenceEditorSubsystem.cpp b/Engine/Plugins/MovieScene/LevelSequenceEditor/Source/LevelSequenceEditor/Private/LevelSequenceEditorSubsystem.cpp
--- a/Engine/Plugins/MovieScene/LevelSequenceEditor/Source/LevelSequenceEditor/Private/LevelSequenceEditorSubsystem.cpp
+++ b/Engine/Plugins/MovieScene/LevelSequenceEditor/Source/LevelSequenceEditor/Private/LevelSequenceEditorSubsystem.cpp
@@ -30,8 +30,11 @@
 #include "Framework/MultiBox/MultiBoxBuilder.h"
 #include "Framework/Notifications/NotificationManager.h"
 #include "GameFramework/Actor.h"
+#include "IContentBrowserSingleton.h"
 #include "LevelEditor.h"
 #include "Modules/ModuleManager.h"
+#include "SceneOutlinerModule.h"
+#include "SceneOutlinerPublicTypes.h"
 #include "ScopedTransaction.h"
 #include "Widgets/Notifications/SNotificationList.h"
 
@@ -39,15 +42,6 @@ DEFINE_LOG_CATEGORY(LogLevelSequenceEditor);
 
 #define LOCTEXT_NAMESPACE "LevelSequenceEditor"
 
-void AddAssignActorMenu(FMenuBuilder& MenuBuilder)
-{
-	MenuBuilder.AddMenuEntry(FLevelSequenceEditorCommands::Get().AddActorsToBinding);
-	MenuBuilder.AddMenuEntry(FLevelSequenceEditorCommands::Get().ReplaceBindingWithActors);
-	MenuBuilder.AddMenuEntry(FLevelSequenceEditorCommands::Get().RemoveActorsFromBinding);
-	MenuBuilder.AddMenuEntry(FLevelSequenceEditorCommands::Get().RemoveAllBindings);
-	MenuBuilder.AddMenuEntry(FLevelSequenceEditorCommands::Get().RemoveInvalidBindings);
-}
-
 void ULevelSequenceEditorSubsystem::Initialize(FSubsystemCollectionBase& Collection)
 {
 	UE_LOG(LogLevelSequenceEditor, Log, TEXT("LevelSequenceEditor subsystem initialized."));
@@ -130,12 +124,12 @@ void ULevelSequenceEditorSubsystem::Initialize(FSubsystemCollectionBase& Collect
 	SequencerModule.GetActionsMenuExtensibilityManager()->AddExtender(FixActorReferencesMenuExtender);
 
 	AssignActorMenuExtender = MakeShareable(new FExtender);
-	AssignActorMenuExtender->AddMenuExtension("Possessable", EExtensionHook::First, CommandList, FMenuExtensionDelegate::CreateStatic([](FMenuBuilder& MenuBuilder) {
+	AssignActorMenuExtender->AddMenuExtension("Possessable", EExtensionHook::First, CommandList, FMenuExtensionDelegate::CreateLambda([this](FMenuBuilder& MenuBuilder) {
 		FFormatNamedArguments Args;
 		MenuBuilder.AddSubMenu(
 			FText::Format(LOCTEXT("AssignActor", "Assign Actor"), Args),
 			FText::Format(LOCTEXT("AssignActorTooltip", "Assign an actor to this track"), Args),
-			FNewMenuDelegate::CreateStatic(&AddAssignActorMenu));
+			FNewMenuDelegate::CreateLambda([this](FMenuBuilder& SubMenuBuilder) { AddAssignActorMenu(SubMenuBuilder); } ));
 		}));
 
 	SequencerModule.GetObjectBindingContextMenuExtensibilityManager()->AddExtender(AssignActorMenuExtender);
@@ -1264,6 +1258,80 @@ void ULevelSequenceEditorSubsystem::RemoveInvalidBindings(const FSequencerBindin
 	Sequencer->NotifyMovieSceneDataChanged(EMovieSceneDataChangeType::MovieSceneStructureItemsChanged);
 }
 
+void ULevelSequenceEditorSubsystem::AddAssignActorMenu(FMenuBuilder& MenuBuilder)
+{
+	MenuBuilder.AddMenuEntry(FLevelSequenceEditorCommands::Get().AddActorsToBinding);
+	MenuBuilder.AddMenuEntry(FLevelSequenceEditorCommands::Get().ReplaceBindingWithActors);
+	MenuBuilder.AddMenuEntry(FLevelSequenceEditorCommands::Get().RemoveActorsFromBinding);
+	MenuBuilder.AddMenuEntry(FLevelSequenceEditorCommands::Get().RemoveAllBindings);
+	MenuBuilder.AddMenuEntry(FLevelSequenceEditorCommands::Get().RemoveInvalidBindings);
+
+	TSharedPtr<ISequencer> Sequencer = GetActiveSequencer();
+	if (Sequencer == nullptr)
+	{
+		return;
+	}
+
+	TArray<FGuid> ObjectBindings;
+	Sequencer->GetSelectedObjects(ObjectBindings);
+	if (ObjectBindings.Num() == 0)
+	{
+		return;
+	}
+
+	FSequencerBindingProxy BindingProxy(ObjectBindings[0], Sequencer->GetFocusedMovieSceneSequence());
+
+	TSet<const AActor*> BoundObjects;
+	{
+		for (TWeakObjectPtr<> Ptr : Sequencer->FindObjectsInCurrentSequence(ObjectBindings[0]))
+		{
+			if (const AActor* Actor = Cast<AActor>(Ptr.Get()))
+			{
+				BoundObjects.Add(Actor);
+			}
+		}
+	}
+
+	auto IsActorValidForAssignment = [BoundObjects](const AActor* InActor){
+		return !BoundObjects.Contains(InActor);
+	};
+
+	// Set up a menu entry to assign an actor to the object binding node
+	FSceneOutlinerInitializationOptions InitOptions;
+	{
+		// We hide the header row to keep the UI compact.
+		InitOptions.bShowHeaderRow = false;
+		InitOptions.bShowSearchBox = true;
+		InitOptions.bShowCreateNewFolder = false;
+		InitOptions.bFocusSearchBoxWhenOpened = true;
+		// Only want the actor label column
+		InitOptions.ColumnMap.Add(FSceneOutlinerBuiltInColumnTypes::Label(), FSceneOutlinerColumnInfo(ESceneOutlinerColumnVisibility::Visible, 0));
+		
+		// Only display actors that are not possessed already
+		InitOptions.Filters->AddFilterPredicate<FActorTreeItem>(FActorTreeItem::FFilterPredicate::CreateLambda( IsActorValidForAssignment ) );
+	}
+
+	// actor selector to allow the user to choose an actor
+	FSceneOutlinerModule& SceneOutlinerModule = FModuleManager::LoadModuleChecked<FSceneOutlinerModule>("SceneOutliner");
+	TSharedRef< SWidget > MiniSceneOutliner =
+		SNew( SBox )
+		.MaxDesiredHeight(400.0f)
+		.WidthOverride(300.0f)
+		[
+			SceneOutlinerModule.CreateActorPicker(
+				InitOptions,
+				FOnActorPicked::CreateLambda([=](AActor* Actor){
+					// Create a new binding for this actor
+					FSlateApplication::Get().DismissAllMenus();
+					FSequencerUtilities::DoAssignActor(Sequencer.Get(), Actor, ObjectBindings[0]);
+				})
+			)
+		];
+
+	MenuBuilder.AddMenuSeparator();
+	MenuBuilder.AddWidget(MiniSceneOutliner, FText::GetEmpty(), true);
+}
+
 void ULevelSequenceEditorSubsystem::GetRebindComponentNames(TArray<FName>& OutComponentNames)
 {
 	OutComponentNames.Empty();
diff --git a/Engine/Plugins/MovieScene/LevelSequenceEditor/Source/LevelSequenceEditor/Public/LevelSequenceEditorSubsystem.h b/Engine/Plugins/MovieScene/LevelSequenceEditor/Source/LevelSequenceEditor/Public/LevelSequenceEditorSubsystem.h
--- a/Engine/Plugins/MovieScene/LevelSequenceEditor/Source/LevelSequenceEditor/Public/LevelSequenceEditorSubsystem.h
+++ b/Engine/Plugins/MovieScene/LevelSequenceEditor/Source/LevelSequenceEditor/Public/LevelSequenceEditorSubsystem.h
@@ -99,6 +99,8 @@ private:
 	void RemoveInvalidBindingsInternal();
 	void RebindComponentInternal(const FName& ComponentName);
 
+	void AddAssignActorMenu(FMenuBuilder& MenuBuilder);
+
 	void GetRebindComponentNames(TArray<FName>& OutComponentNames);
 	void RebindComponentMenu(FMenuBuilder& MenuBuilder);
 
diff --git a/Engine/Source/Editor/Sequencer/Private/Sequencer.cpp b/Engine/Source/Editor/Sequencer/Private/Sequencer.cpp
--- a/Engine/Source/Editor/Sequencer/Private/Sequencer.cpp
+++ b/Engine/Source/Editor/Sequencer/Private/Sequencer.cpp
@@ -7443,60 +7443,6 @@ void FSequencer::DeleteSelectedItems()
 }
 
 
-void FSequencer::AssignActor(FMenuBuilder& MenuBuilder, FGuid InObjectBinding)
-{
-	TSet<const AActor*> BoundObjects;
-	{
-		for (TWeakObjectPtr<> Ptr : FindObjectsInCurrentSequence(InObjectBinding))
-		{
-			if (const AActor* Actor = Cast<AActor>(Ptr.Get()))
-			{
-				BoundObjects.Add(Actor);
-			}
-		}
-	}
-
-	auto IsActorValidForAssignment = [BoundObjects](const AActor* InActor){
-		return !BoundObjects.Contains(InActor);
-	};
-
-	// Set up a menu entry to assign an actor to the object binding node
-	FSceneOutlinerInitializationOptions InitOptions;
-	{
-		// We hide the header row to keep the UI compact.
-		InitOptions.bShowHeaderRow = false;
-		InitOptions.bShowSearchBox = true;
-		InitOptions.bShowCreateNewFolder = false;
-		InitOptions.bFocusSearchBoxWhenOpened = true;
-		// Only want the actor label column
-		InitOptions.ColumnMap.Add(FSceneOutlinerBuiltInColumnTypes::Label(), FSceneOutlinerColumnInfo(ESceneOutlinerColumnVisibility::Visible, 0));
-		
-		// Only display actors that are not possessed already
-		InitOptions.Filters->AddFilterPredicate<FActorTreeItem>(FActorTreeItem::FFilterPredicate::CreateLambda( IsActorValidForAssignment ) );
-	}
-
-	// actor selector to allow the user to choose an actor
-	FSceneOutlinerModule& SceneOutlinerModule = FModuleManager::LoadModuleChecked<FSceneOutlinerModule>("SceneOutliner");
-	TSharedRef< SWidget > MiniSceneOutliner =
-		SNew( SBox )
-		.MaxDesiredHeight(400.0f)
-		.WidthOverride(300.0f)
-		[
-			SceneOutlinerModule.CreateActorPicker(
-				InitOptions,
-				FOnActorPicked::CreateLambda([=](AActor* Actor){
-					// Create a new binding for this actor
-					FSlateApplication::Get().DismissAllMenus();
-					FSequencerUtilities::DoAssignActor(this, Actor, InObjectBinding);
-				})
-			)
-		];
-
-	MenuBuilder.AddMenuSeparator();
-	MenuBuilder.AddWidget(MiniSceneOutliner, FText::GetEmpty(), true);
-}
-
-
 void FSequencer::AddActorsToBinding(FGuid InObjectBinding, const TArray<AActor*>& InActors)
 {
 	if (!InActors.Num())
diff --git a/Engine/Source/Editor/Sequencer/Private/Sequencer.h b/Engine/Source/Editor/Sequencer/Private/Sequencer.h
--- a/Engine/Source/Editor/Sequencer/Private/Sequencer.h
+++ b/Engine/Source/Editor/Sequencer/Private/Sequencer.h
@@ -518,9 +518,6 @@ public:
 	FText GetNavigateForwardTooltip() const;
 	FText GetNavigateBackwardTooltip() const;
 
-	/** Called when a user executes the assign actor to track menu item */
-	void AssignActor(FMenuBuilder& MenuBuilder, FGuid ObjectBinding);
-
 	/** Called when a user executes the assign selected to track menu item */
 	void AddActorsToBinding(FGuid ObjectBinding, const TArray<AActor*>& InActors);
 	void ReplaceBindingWithActors(FGuid ObjectBinding, const TArray<AActor*>& InActors);
