Non-unity fix

diff --git a/Engine/Plugins/MovieScene/LevelSequenceEditor/Source/LevelSequenceEditor/Private/LevelSequenceEditorSubsystem.cpp b/Engine/Plugins/MovieScene/LevelSequenceEditor/Source/LevelSequenceEditor/Private/LevelSequenceEditorSubsystem.cpp
--- a/Engine/Plugins/MovieScene/LevelSequenceEditor/Source/LevelSequenceEditor/Private/LevelSequenceEditorSubsystem.cpp
+++ b/Engine/Plugins/MovieScene/LevelSequenceEditor/Source/LevelSequenceEditor/Private/LevelSequenceEditorSubsystem.cpp
@@ -37,7 +37,7 @@
 #include "SceneOutlinerPublicTypes.h"
 #include "ScopedTransaction.h"
 #include "Widgets/Notifications/SNotificationList.h"
-
+#include "Framework/MultiBox/MultiBoxBuilder.h"
 DEFINE_LOG_CATEGORY(LogLevelSequenceEditor);
 
 #define LOCTEXT_NAMESPACE "LevelSequenceEditor"
