Sequencer: Check for valid root component when converting to spawnable

diff --git a/Engine/Plugins/MovieScene/LevelSequenceEditor/Source/LevelSequenceEditor/Private/Misc/LevelSequenceEditorSpawnRegister.cpp b/Engine/Plugins/MovieScene/LevelSequenceEditor/Source/LevelSequenceEditor/Private/Misc/LevelSequenceEditorSpawnRegister.cpp
--- a/Engine/Plugins/MovieScene/LevelSequenceEditor/Source/LevelSequenceEditor/Private/Misc/LevelSequenceEditorSpawnRegister.cpp
+++ b/Engine/Plugins/MovieScene/LevelSequenceEditor/Source/LevelSequenceEditor/Private/Misc/LevelSequenceEditorSpawnRegister.cpp
@@ -340,10 +340,13 @@ void FLevelSequenceEditorSpawnRegister::HandleConvertPossessableToSpawnable(UObj
 	AActor* OldActor = Cast<AActor>(OldObject);
 	if (OldActor)
 	{
-		OutTransformData.Emplace();
-		OutTransformData->Translation = OldActor->GetRootComponent()->GetRelativeLocation();
-		OutTransformData->Rotation = OldActor->GetRootComponent()->GetRelativeRotation();
-		OutTransformData->Scale = OldActor->GetRootComponent()->GetRelativeScale3D();
+		if (OldActor->GetRootComponent())
+		{
+			OutTransformData.Emplace();
+			OutTransformData->Translation = OldActor->GetRootComponent()->GetRelativeLocation();
+			OutTransformData->Rotation = OldActor->GetRootComponent()->GetRelativeRotation();
+			OutTransformData->Scale = OldActor->GetRootComponent()->GetRelativeScale3D();
+		}
 
 		GEditor->SelectActor(OldActor, false, true);
 		UWorld* World = Cast<UWorld>(Player.GetPlaybackContext());
