Marked up LevelSequenceEditorSubsystem with _API macro to allow for usage outside of the module

diff --git a/Engine/Plugins/MovieScene/LevelSequenceEditor/Source/LevelSequenceEditor/Public/LevelSequenceEditorSubsystem.h b/Engine/Plugins/MovieScene/LevelSequenceEditor/Source/LevelSequenceEditor/Public/LevelSequenceEditorSubsystem.h
--- a/Engine/Plugins/MovieScene/LevelSequenceEditor/Source/LevelSequenceEditor/Public/LevelSequenceEditorSubsystem.h
+++ b/Engine/Plugins/MovieScene/LevelSequenceEditor/Source/LevelSequenceEditor/Public/LevelSequenceEditorSubsystem.h
@@ -34,7 +34,7 @@ struct FMovieSceneScriptingParams
 * Subsystem for level sequencer related utilities to scripts
 */
 UCLASS(Blueprintable)
-class ULevelSequenceEditorSubsystem : public UEditorSubsystem
+class LEVELSEQUENCEEDITOR_API ULevelSequenceEditorSubsystem : public UEditorSubsystem
 {
 	GENERATED_BODY()
 
