Sequencer: Add asterisk to the Sequencer tab if the asset is dirty (similar to level editor)

diff --git a/Engine/Plugins/MovieScene/LevelSequenceEditor/Source/LevelSequenceEditor/Private/LevelSequenceEditorToolkit.cpp b/Engine/Plugins/MovieScene/LevelSequenceEditor/Source/LevelSequenceEditor/Private/LevelSequenceEditorToolkit.cpp
--- a/Engine/Plugins/MovieScene/LevelSequenceEditor/Source/LevelSequenceEditor/Private/LevelSequenceEditorToolkit.cpp
+++ b/Engine/Plugins/MovieScene/LevelSequenceEditor/Source/LevelSequenceEditor/Private/LevelSequenceEditorToolkit.cpp
@@ -228,7 +228,12 @@ void FLevelSequenceEditorToolkit::Initialize(const EToolkitMode::Type Mode, cons
 	}
 	
 	// Now Attach so this window will apear in the correct front first order
-	LevelEditorModule.AttachSequencer(Sequencer->GetSequencerWidget(), SharedThis(this));
+	TSharedPtr<SDockTab> DockTab = LevelEditorModule.AttachSequencer(Sequencer->GetSequencerWidget(), SharedThis(this));
+	if (DockTab.IsValid())
+	{
+		TAttribute<FText> LabelSuffix = TAttribute<FText>(this, &FLevelSequenceEditorToolkit::GetTabSuffix);
+		DockTab->SetTabLabelSuffix(LabelSuffix);
+	}
 
 	// We need to find out when the user loads a new map, because we might need to re-create puppet actors
 	// when previewing a MovieScene
@@ -287,6 +292,24 @@ FString FLevelSequenceEditorToolkit::GetWorldCentricTabPrefix() const
 	return LOCTEXT("WorldCentricTabPrefix", "Sequencer ").ToString();
 }
 
+FText FLevelSequenceEditorToolkit::GetTabSuffix() const
+{
+	UMovieSceneSequence* Sequence = Sequencer->GetFocusedMovieSceneSequence();
+
+	if (Sequence == nullptr)
+	{
+		return FText::GetEmpty();
+	}
+	
+	const bool bIsDirty = Sequence->GetMovieScene()->GetOuter()->GetOutermost()->IsDirty();
+	if (bIsDirty)
+	{
+		return LOCTEXT("TabSuffixAsterix", "*");
+	}
+
+	return FText::GetEmpty();
+}
+
 /* FLevelSequenceEditorToolkit implementation
  *****************************************************************************/
 
diff --git a/Engine/Plugins/MovieScene/LevelSequenceEditor/Source/LevelSequenceEditor/Private/LevelSequenceEditorToolkit.h b/Engine/Plugins/MovieScene/LevelSequenceEditor/Source/LevelSequenceEditor/Private/LevelSequenceEditorToolkit.h
--- a/Engine/Plugins/MovieScene/LevelSequenceEditor/Source/LevelSequenceEditor/Private/LevelSequenceEditorToolkit.h
+++ b/Engine/Plugins/MovieScene/LevelSequenceEditor/Source/LevelSequenceEditor/Private/LevelSequenceEditorToolkit.h
@@ -104,6 +104,7 @@ public:
 	virtual FName GetToolkitFName() const override;
 	virtual FLinearColor GetWorldCentricTabColorScale() const override;
 	virtual FString GetWorldCentricTabPrefix() const override;
+	virtual FText GetTabSuffix() const override;
 
 protected:
 
diff --git a/Engine/Plugins/MovieScene/TemplateSequence/Source/TemplateSequenceEditor/Private/TemplateSequenceEditorToolkit.cpp b/Engine/Plugins/MovieScene/TemplateSequence/Source/TemplateSequenceEditor/Private/TemplateSequenceEditorToolkit.cpp
--- a/Engine/Plugins/MovieScene/TemplateSequence/Source/TemplateSequenceEditor/Private/TemplateSequenceEditorToolkit.cpp
+++ b/Engine/Plugins/MovieScene/TemplateSequence/Source/TemplateSequenceEditor/Private/TemplateSequenceEditorToolkit.cpp
@@ -127,7 +127,13 @@ void FTemplateSequenceEditorToolkit::Initialize(const EToolkitMode::Type Mode, c
 		}
 	}
 
-	LevelEditorModule.AttachSequencer(Sequencer->GetSequencerWidget(), SharedThis(this));
+	TSharedPtr<SDockTab> DockTab = LevelEditorModule.AttachSequencer(Sequencer->GetSequencerWidget(), SharedThis(this));
+	if (DockTab.IsValid())
+	{
+		TAttribute<FText> LabelSuffix = TAttribute<FText>(this, &FTemplateSequenceEditorToolkit::GetTabSuffix);
+		DockTab->SetTabLabelSuffix(LabelSuffix);
+	}
+
 	LevelEditorModule.OnMapChanged().AddRaw(this, &FTemplateSequenceEditorToolkit::HandleMapChanged);
 }
 
@@ -147,6 +153,24 @@ FString FTemplateSequenceEditorToolkit::GetWorldCentricTabPrefix() const
 	return LOCTEXT("WorldCentricTabPrefix", "Sequencer ").ToString();
 }
 
+FText FTemplateSequenceEditorToolkit::GetTabSuffix() const
+{
+	UMovieSceneSequence* Sequence = Sequencer->GetFocusedMovieSceneSequence();
+
+	if (Sequence == nullptr)
+	{
+		return FText::GetEmpty();
+	}
+
+	const bool bIsDirty = Sequence->GetMovieScene()->GetOuter()->GetOutermost()->IsDirty();
+	if (bIsDirty)
+	{
+		return LOCTEXT("TabSuffixAsterix", "*");
+	}
+
+	return FText::GetEmpty();
+}
+
 FLinearColor FTemplateSequenceEditorToolkit::GetWorldCentricTabColorScale() const
 {
 	return FLinearColor(0.7, 0.0f, 0.0f, 0.5f);
diff --git a/Engine/Plugins/MovieScene/TemplateSequence/Source/TemplateSequenceEditor/Private/TemplateSequenceEditorToolkit.h b/Engine/Plugins/MovieScene/TemplateSequence/Source/TemplateSequenceEditor/Private/TemplateSequenceEditorToolkit.h
--- a/Engine/Plugins/MovieScene/TemplateSequence/Source/TemplateSequenceEditor/Private/TemplateSequenceEditorToolkit.h
+++ b/Engine/Plugins/MovieScene/TemplateSequence/Source/TemplateSequenceEditor/Private/TemplateSequenceEditorToolkit.h
@@ -70,6 +70,7 @@ public:
 	virtual FText GetBaseToolkitName() const override;
 	virtual FName GetToolkitFName() const override;
 	virtual FString GetWorldCentricTabPrefix() const override;
+	virtual FText GetTabSuffix() const override;
 	virtual FLinearColor GetWorldCentricTabColorScale() const override;
 	virtual void RegisterTabSpawners(const TSharedRef<FTabManager>& TabManager) override;
 	virtual void UnregisterTabSpawners(const TSharedRef<FTabManager>& TabManager) override;
diff --git a/Engine/Source/Editor/LevelEditor/Private/LevelEditor.cpp b/Engine/Source/Editor/LevelEditor/Private/LevelEditor.cpp
--- a/Engine/Source/Editor/LevelEditor/Private/LevelEditor.cpp
+++ b/Engine/Source/Editor/LevelEditor/Private/LevelEditor.cpp
@@ -379,13 +379,15 @@ void FLevelEditorModule::SummonWorldBrowserComposition()
 }
 
 // @todo remove when world-centric mode is added
-void FLevelEditorModule::AttachSequencer(TSharedPtr<SWidget> SequencerWidget, TSharedPtr<IAssetEditorInstance> SequencerAssetEditor)
+TSharedPtr<SDockTab> FLevelEditorModule::AttachSequencer(TSharedPtr<SWidget> SequencerWidget, TSharedPtr<IAssetEditorInstance> SequencerAssetEditor)
 {
 	TSharedPtr<SLevelEditor> LevelEditorInstance = LevelEditorInstancePtr.Pin();
 	if (LevelEditorInstance)
 	{
-		LevelEditorInstance->AttachSequencer(SequencerWidget, SequencerAssetEditor);
+		return LevelEditorInstance->AttachSequencer(SequencerWidget, SequencerAssetEditor);
 	}
+
+	return nullptr;
 }
 
 TSharedPtr<IAssetViewport> FLevelEditorModule::GetFirstActiveViewport()
diff --git a/Engine/Source/Editor/LevelEditor/Private/SLevelEditor.cpp b/Engine/Source/Editor/LevelEditor/Private/SLevelEditor.cpp
--- a/Engine/Source/Editor/LevelEditor/Private/SLevelEditor.cpp
+++ b/Engine/Source/Editor/LevelEditor/Private/SLevelEditor.cpp
@@ -594,7 +594,7 @@ TSharedPtr<FTabManager> SLevelEditor::GetTabManager() const
 	return LevelEditorTabManager;
 }
 
-void SLevelEditor::AttachSequencer( TSharedPtr<SWidget> SequencerWidget, TSharedPtr<IAssetEditorInstance> NewSequencerAssetEditor )
+TSharedPtr<SDockTab> SLevelEditor::AttachSequencer( TSharedPtr<SWidget> SequencerWidget, TSharedPtr<IAssetEditorInstance> NewSequencerAssetEditor )
 {
 	struct Local
 	{
@@ -649,7 +649,11 @@ void SLevelEditor::AttachSequencer( TSharedPtr<SWidget> SequencerWidget, TShared
 				SequencerAssetEditor.Reset();
 			}
 		}
+
+		return Tab;
 	}
+	
+	return nullptr;
 }
 
 TArray<TWeakPtr<ISceneOutliner>> SLevelEditor::GetAllSceneOutliners() const
diff --git a/Engine/Source/Editor/LevelEditor/Private/SLevelEditor.h b/Engine/Source/Editor/LevelEditor/Private/SLevelEditor.h
--- a/Engine/Source/Editor/LevelEditor/Private/SLevelEditor.h
+++ b/Engine/Source/Editor/LevelEditor/Private/SLevelEditor.h
@@ -21,6 +21,7 @@ class IAssetEditorInstance;
 class IDetailsView;
 class SActorDetails;
 class SBorder;
+class SDockTab;
 class SLevelEditorModeContent;
 class SLevelEditorToolBox;
 class UTypedElementSelectionSet;
@@ -148,7 +149,7 @@ public:
 	}
 	
 	/** Attaches a sequencer asset editor used to animate objects in the level to this level editor */
-	void AttachSequencer( TSharedPtr<SWidget> SequencerWidget, TSharedPtr<IAssetEditorInstance> NewSequencerAssetEditor );
+	TSharedPtr<SDockTab> AttachSequencer( TSharedPtr<SWidget> SequencerWidget, TSharedPtr<IAssetEditorInstance> NewSequencerAssetEditor );
 
 	/** Get an array containing weak pointers to all 4 Scene Outliners which could be potentially active */
 	virtual TArray<TWeakPtr<ISceneOutliner>> GetAllSceneOutliners() const override;
diff --git a/Engine/Source/Editor/LevelEditor/Public/LevelEditor.h b/Engine/Source/Editor/LevelEditor/Public/LevelEditor.h
--- a/Engine/Source/Editor/LevelEditor/Public/LevelEditor.h
+++ b/Engine/Source/Editor/LevelEditor/Public/LevelEditor.h
@@ -15,6 +15,7 @@
 
 class AActor;
 class IAssetViewport;
+class SDockTab;
 class SLevelEditor;
 class UAnimSequence;
 class USkeletalMeshComponent;
@@ -112,7 +113,7 @@ public:
 	 * Spawns a new sequencer tab if one doesn't exist already
 	 * @todo This only works with the first level editor. Fix it.
 	 */
-	virtual void AttachSequencer(TSharedPtr<SWidget> SequencerWidget, TSharedPtr<class IAssetEditorInstance> SequencerAssetEditor );
+	virtual TSharedPtr<SDockTab> AttachSequencer(TSharedPtr<SWidget> SequencerWidget, TSharedPtr<class IAssetEditorInstance> SequencerAssetEditor );
 
 	/**
 	 * Starts a play in editor session using the active viewport
