Fixed derefencing null pointer during fix actor referneces

diff --git a/Engine/Plugins/MovieScene/LevelSequenceEditor/Source/LevelSequenceEditor/Private/LevelSequenceEditorSubsystem.cpp b/Engine/Plugins/MovieScene/LevelSequenceEditor/Source/LevelSequenceEditor/Private/LevelSequenceEditorSubsystem.cpp
--- a/Engine/Plugins/MovieScene/LevelSequenceEditor/Source/LevelSequenceEditor/Private/LevelSequenceEditorSubsystem.cpp
+++ b/Engine/Plugins/MovieScene/LevelSequenceEditor/Source/LevelSequenceEditor/Private/LevelSequenceEditorSubsystem.cpp
@@ -801,7 +801,7 @@ void ULevelSequenceEditorSubsystem::FixActorReferences()
 	TMap<FGuid, FGuid> OldGuidToNewGuidMap;
 	for (const FMovieScenePossessable& ActorPossessableToFix : ActorsPossessablesToFix)
 	{
-		AActor* ActorPtr = *ActorNameToActorMap.Find(ActorPossessableToFix.GetName());
+		AActor* ActorPtr = ActorNameToActorMap.FindRef(ActorPossessableToFix.GetName());
 		if (ActorPtr != nullptr)
 		{
 			FGuid OldGuid = ActorPossessableToFix.GetGuid();
