Sequencer: Close existing level editor sequencer instance before opening a new one when switching assets

diff --git a/Engine/Plugins/MovieScene/LevelSequenceEditor/Source/LevelSequenceEditor/Private/LevelSequenceEditorToolkit.cpp b/Engine/Plugins/MovieScene/LevelSequenceEditor/Source/LevelSequenceEditor/Private/LevelSequenceEditorToolkit.cpp
--- a/Engine/Plugins/MovieScene/LevelSequenceEditor/Source/LevelSequenceEditor/Private/LevelSequenceEditorToolkit.cpp
+++ b/Engine/Plugins/MovieScene/LevelSequenceEditor/Source/LevelSequenceEditor/Private/LevelSequenceEditorToolkit.cpp
@@ -156,6 +156,11 @@ FLevelSequenceEditorToolkit::~FLevelSequenceEditorToolkit()
 
 void FLevelSequenceEditorToolkit::Initialize(const EToolkitMode::Type Mode, const TSharedPtr<IToolkitHost>& InitToolkitHost, ULevelSequence* InLevelSequence)
 {
+	FLevelEditorModule& LevelEditorModule = FModuleManager::LoadModuleChecked<FLevelEditorModule>("LevelEditor");
+
+	// Clear out the existing sequencer
+	LevelEditorModule.AttachSequencer(nullptr, nullptr);
+
 	// create tab layout
 	const TSharedRef<FTabManager::FLayout> StandaloneDefaultLayout = FTabManager::NewLayout("Standalone_LevelSequenceEditor")
 		->AddArea
@@ -214,9 +219,6 @@ void FLevelSequenceEditorToolkit::Initialize(const EToolkitMode::Type Mode, cons
 	FLevelEditorSequencerIntegration::Get().AddSequencer(Sequencer.ToSharedRef(), Options);
 	ULevelSequenceEditorBlueprintLibrary::SetSequencer(Sequencer.ToSharedRef());
 
-	// @todo remove when world-centric mode is added
-	FLevelEditorModule& LevelEditorModule = FModuleManager::LoadModuleChecked<FLevelEditorModule>("LevelEditor");
-
 	// Reopen the scene outliner so that is refreshed with the sequencer columns
 	{
 		TSharedPtr<FTabManager> LevelEditorTabManager = LevelEditorModule.GetLevelEditorTabManager();
