LevelSequenceEditor: Restrict sequencer editor menu extensions to focused level sequences.

diff --git a/Engine/Plugins/MovieScene/LevelSequenceEditor/Source/LevelSequenceEditor/Private/LevelSequenceEditorSubsystem.cpp b/Engine/Plugins/MovieScene/LevelSequenceEditor/Source/LevelSequenceEditor/Private/LevelSequenceEditorSubsystem.cpp
--- a/Engine/Plugins/MovieScene/LevelSequenceEditor/Source/LevelSequenceEditor/Private/LevelSequenceEditorSubsystem.cpp
+++ b/Engine/Plugins/MovieScene/LevelSequenceEditor/Source/LevelSequenceEditor/Private/LevelSequenceEditorSubsystem.cpp
@@ -110,7 +110,13 @@ void ULevelSequenceEditorSubsystem::Initialize(FSubsystemCollectionBase& Collect
 
 	/* Menu extenders */
 	TransformMenuExtender = MakeShareable(new FExtender);
-	TransformMenuExtender->AddMenuExtension("Transform", EExtensionHook::After, CommandList, FMenuExtensionDelegate::CreateStatic([](FMenuBuilder& MenuBuilder) {
+	TransformMenuExtender->AddMenuExtension("Transform", EExtensionHook::After, CommandList, FMenuExtensionDelegate::CreateLambda([this](FMenuBuilder& MenuBuilder) {
+		// Only add menu entries where the focused sequence is a ULevelSequence
+		if (!GetActiveSequencer())
+		{
+			return;
+		}
+		
 		MenuBuilder.AddMenuEntry(FLevelSequenceEditorCommands::Get().SnapSectionsToTimelineUsingSourceTimecode);
 		MenuBuilder.AddMenuEntry(FLevelSequenceEditorCommands::Get().SyncSectionsUsingSourceTimecode);
 		MenuBuilder.AddMenuEntry(FLevelSequenceEditorCommands::Get().BakeTransform);
@@ -119,7 +125,13 @@ void ULevelSequenceEditorSubsystem::Initialize(FSubsystemCollectionBase& Collect
 	SequencerModule.GetActionsMenuExtensibilityManager()->AddExtender(TransformMenuExtender);
 
 	FixActorReferencesMenuExtender = MakeShareable(new FExtender);
-	FixActorReferencesMenuExtender->AddMenuExtension("Bindings", EExtensionHook::First, CommandList, FMenuExtensionDelegate::CreateStatic([](FMenuBuilder& MenuBuilder) {
+	FixActorReferencesMenuExtender->AddMenuExtension("Bindings", EExtensionHook::First, CommandList, FMenuExtensionDelegate::CreateLambda([this](FMenuBuilder& MenuBuilder) {
+		// Only add menu entries where the focused sequence is a ULevelSequence
+		if (!GetActiveSequencer())
+		{
+			return;
+		}
+		
 		MenuBuilder.AddMenuEntry(FLevelSequenceEditorCommands::Get().FixActorReferences);
 		}));
 
@@ -127,6 +139,12 @@ void ULevelSequenceEditorSubsystem::Initialize(FSubsystemCollectionBase& Collect
 
 	AssignActorMenuExtender = MakeShareable(new FExtender);
 	AssignActorMenuExtender->AddMenuExtension("Possessable", EExtensionHook::First, CommandList, FMenuExtensionDelegate::CreateLambda([this](FMenuBuilder& MenuBuilder) {
+		// Only add menu entries where the focused sequence is a ULevelSequence
+		if (!GetActiveSequencer())
+		{
+			return;
+		}
+		
 		FFormatNamedArguments Args;
 		MenuBuilder.AddSubMenu(
 			FText::Format(LOCTEXT("AssignActor", "Assign Actor"), Args),
@@ -138,6 +156,12 @@ void ULevelSequenceEditorSubsystem::Initialize(FSubsystemCollectionBase& Collect
 
 	RebindComponentMenuExtender = MakeShareable(new FExtender);
 	RebindComponentMenuExtender->AddMenuExtension("Possessable", EExtensionHook::First, CommandList, FMenuExtensionDelegate::CreateLambda([this](FMenuBuilder& MenuBuilder) {
+		// Only add menu entries where the focused sequence is a ULevelSequence
+		if (!GetActiveSequencer())
+		{
+			return;
+		}
+		
 		TArray<FName> ComponentNames;
 		GetRebindComponentNames(ComponentNames);
 		if (ComponentNames.Num() > 0)
