Expose movie scene sub section / focused sequence navigation to scripting - focus section / focus parent / get hierarchy

diff --git a/Engine/Plugins/MovieScene/LevelSequenceEditor/Source/LevelSequenceEditor/Private/LevelSequenceEditorBlueprintLibrary.cpp b/Engine/Plugins/MovieScene/LevelSequenceEditor/Source/LevelSequenceEditor/Private/LevelSequenceEditorBlueprintLibrary.cpp
--- a/Engine/Plugins/MovieScene/LevelSequenceEditor/Source/LevelSequenceEditor/Private/LevelSequenceEditorBlueprintLibrary.cpp
+++ b/Engine/Plugins/MovieScene/LevelSequenceEditor/Source/LevelSequenceEditor/Private/LevelSequenceEditorBlueprintLibrary.cpp
@@ -19,6 +19,8 @@
 //For custom colors on channels, stored in editor pref's
 #include "CurveEditorSettings.h"
 
+#include "Sections/MovieSceneSubSection.h"
+
 namespace
 {
 	static TWeakPtr<ISequencer> CurrentSequencer;
@@ -52,6 +54,47 @@ ULevelSequence* ULevelSequenceEditorBlueprintLibrary::GetFocusedLevelSequence()
 	return nullptr;
 }
 
+void ULevelSequenceEditorBlueprintLibrary::FocusLevelSequence(UMovieSceneSubSection* SubSection)
+{
+	if (CurrentSequencer.IsValid() && IsValid(SubSection))
+	{
+		CurrentSequencer.Pin()->FocusSequenceInstance(*SubSection);
+	}
+}
+
+void ULevelSequenceEditorBlueprintLibrary::FocusParentSequence()
+{
+	if (CurrentSequencer.IsValid())
+	{
+		const TArray<FMovieSceneSequenceID>& Hierarchy = CurrentSequencer.Pin()->GetSubSequenceHierarchy();
+
+		if (Hierarchy.Num() > 1)
+		{
+			CurrentSequencer.Pin()->PopToSequenceInstance(Hierarchy[Hierarchy.Num() -2]);
+		}
+	}
+}
+
+TArray<UMovieSceneSubSection*> ULevelSequenceEditorBlueprintLibrary::GetSubSequenceHierarchy()
+{
+	TArray<UMovieSceneSubSection*> SectionsHierarchy;
+	
+	if (CurrentSequencer.IsValid())
+	{
+		const TArray<FMovieSceneSequenceID>& Hierarchy = CurrentSequencer.Pin()->GetSubSequenceHierarchy();
+
+		for (const FMovieSceneSequenceID Item : Hierarchy)
+		{
+			if (Item != MovieSceneSequenceID::Root)
+			{
+				// We return the whole hierarchy anyways, leaving it to the user to check for invalid pointers.
+				SectionsHierarchy.Add(CurrentSequencer.Pin()->FindSubSection(Item));
+			}
+		}
+	}
+	return SectionsHierarchy;
+}
+
 void ULevelSequenceEditorBlueprintLibrary::CloseLevelSequence()
 {
 	if (CurrentSequencer.IsValid())
diff --git a/Engine/Plugins/MovieScene/LevelSequenceEditor/Source/LevelSequenceEditor/Public/LevelSequenceEditorBlueprintLibrary.h b/Engine/Plugins/MovieScene/LevelSequenceEditor/Source/LevelSequenceEditor/Public/LevelSequenceEditorBlueprintLibrary.h
--- a/Engine/Plugins/MovieScene/LevelSequenceEditor/Source/LevelSequenceEditor/Public/LevelSequenceEditorBlueprintLibrary.h
+++ b/Engine/Plugins/MovieScene/LevelSequenceEditor/Source/LevelSequenceEditor/Public/LevelSequenceEditorBlueprintLibrary.h
@@ -13,6 +13,7 @@ class ISequencer;
 class ULevelSequence;
 class UMovieSceneFolder;
 class UMovieSceneSection;
+class UMovieSceneSubSection;
 class UMovieSceneTrack;
 
 USTRUCT(BlueprintType)
@@ -62,6 +63,24 @@ public:
 	UFUNCTION(BlueprintCallable, Category = "Level Sequence Editor")
 	static ULevelSequence* GetFocusedLevelSequence();
 
+	/*
+	 * Focus/view the sequence associated to the given sub sequence section.
+	 */
+	UFUNCTION(BlueprintCallable, Category = "Level Sequence Editor")
+	static void FocusLevelSequence(UMovieSceneSubSection* SubSection);
+
+	/*
+	 * Focus/view the parent sequence, popping out of the current sub sequence section.
+	 */
+	UFUNCTION(BlueprintCallable, Category = "Level Sequence Editor")
+	static void FocusParentSequence();
+
+	/*
+	 * Get the current sub section hierarchy from the current sequence to the section associated with the focused sequence.
+	 */
+	UFUNCTION(BlueprintCallable, Category = "Level Sequence Editor")
+	static TArray<UMovieSceneSubSection*> GetSubSequenceHierarchy();
+	
 	/*
 	 * Close
 	 */
diff --git a/Engine/Source/Editor/Sequencer/Private/Sequencer.h b/Engine/Source/Editor/Sequencer/Private/Sequencer.h
--- a/Engine/Source/Editor/Sequencer/Private/Sequencer.h
+++ b/Engine/Source/Editor/Sequencer/Private/Sequencer.h
@@ -242,7 +242,7 @@ public:
 	/**
 	 * Pops the current focused movie scene from the stack.  The parent of this movie scene will be come the focused one
 	 */
-	void PopToSequenceInstance( FMovieSceneSequenceIDRef SequenceID );
+	virtual void PopToSequenceInstance( FMovieSceneSequenceIDRef SequenceID ) override;
 
 	/** Deletes the passed in sections. */
 	void DeleteSections(const TSet<TWeakObjectPtr<UMovieSceneSection> > & Sections);
@@ -681,6 +681,7 @@ public:
 	virtual TSharedRef<SWidget> GetSequencerWidget() const override;
 	virtual FMovieSceneSequenceIDRef GetRootTemplateID() const override { return ActiveTemplateIDs[0]; }
 	virtual FMovieSceneSequenceIDRef GetFocusedTemplateID() const override { return ActiveTemplateIDs.Top(); }
+	virtual const TArray<FMovieSceneSequenceID>& GetSubSequenceHierarchy() const override { return ActiveTemplateIDs; };
 	virtual UMovieSceneSubSection* FindSubSection(FMovieSceneSequenceID SequenceID) const override;
 	virtual UMovieSceneSequence* GetRootMovieSceneSequence() const override;
 	virtual UMovieSceneSequence* GetFocusedMovieSceneSequence() const override;
diff --git a/Engine/Source/Editor/Sequencer/Public/ISequencer.h b/Engine/Source/Editor/Sequencer/Public/ISequencer.h
--- a/Engine/Source/Editor/Sequencer/Public/ISequencer.h
+++ b/Engine/Source/Editor/Sequencer/Public/ISequencer.h
@@ -225,6 +225,7 @@ public:
 	/** @return The root movie scene being used */
 	virtual FMovieSceneSequenceIDRef GetRootTemplateID() const = 0;
 	virtual FMovieSceneSequenceIDRef GetFocusedTemplateID() const = 0;
+	virtual const TArray<FMovieSceneSequenceID>& GetSubSequenceHierarchy() const = 0;
 
 	/** Attempt to locate the sub section that relates to the specified sequence ID. */
 	virtual UMovieSceneSubSection* FindSubSection(FMovieSceneSequenceID SequenceID) const = 0;
@@ -244,6 +245,11 @@ public:
 	 */
 	virtual void FocusSequenceInstance(UMovieSceneSubSection& Section) = 0;
 
+	/**
+	 * Pops the current focused movie scene from the stack.  The parent of this movie scene will be come the focused one
+	 */
+	virtual void PopToSequenceInstance(FMovieSceneSequenceIDRef SequenceID) = 0;
+
 	/**
 	 * Retrieve the top level view model for this sequence
 	 */
