Don't try to resolve bound actors with component possessables.

diff --git a/Engine/Plugins/MovieScene/LevelSequenceEditor/Source/LevelSequenceEditor/Private/Misc/LevelSequenceEditorActorBinding.cpp b/Engine/Plugins/MovieScene/LevelSequenceEditor/Source/LevelSequenceEditor/Private/Misc/LevelSequenceEditorActorBinding.cpp
--- a/Engine/Plugins/MovieScene/LevelSequenceEditor/Source/LevelSequenceEditor/Private/Misc/LevelSequenceEditorActorBinding.cpp
+++ b/Engine/Plugins/MovieScene/LevelSequenceEditor/Source/LevelSequenceEditor/Private/Misc/LevelSequenceEditorActorBinding.cpp
@@ -56,10 +56,13 @@ void FLevelSequenceEditorActorBinding::AddPossessActorMenuExtensions(FMenuBuilde
 		{
 			for (int32 Index = 0; Index < MovieScene->GetPossessableCount(); Index++)
 			{
+				// Only get bound objects for top-level possessables.
 				FMovieScenePossessable& Possessable = MovieScene->GetPossessable(Index);
-
-				// A possession guid can apply to more than one object, so we get all bound objects for the GUID and add them to our set.
-				ExistingPossessedObjects.Append(MovieSceneSequence->LocateBoundObjects(Possessable.GetGuid(), Sequencer.Pin()->GetPlaybackContext()));
+				if (!Possessable.GetParent().IsValid())
+				{
+					// A possession guid can apply to more than one object, so we get all bound objects for the GUID and add them to our set.
+					ExistingPossessedObjects.Append(MovieSceneSequence->LocateBoundObjects(Possessable.GetGuid(), Sequencer.Pin()->GetPlaybackContext()));
+				}
 			}
 		}
 	}
diff --git a/Engine/Source/Runtime/LevelSequence/Private/LevelSequenceBindingReference.cpp b/Engine/Source/Runtime/LevelSequence/Private/LevelSequenceBindingReference.cpp
--- a/Engine/Source/Runtime/LevelSequence/Private/LevelSequenceBindingReference.cpp
+++ b/Engine/Source/Runtime/LevelSequence/Private/LevelSequenceBindingReference.cpp
@@ -146,7 +146,7 @@ UObject* ResolveByPath(UObject* InContext, const FString& InObjectPath)
 		}
 #endif
 
-		if (UObject* FoundObject = FindFirstObject<UObject>(*InObjectPath, EFindFirstObjectOptions::NativeFirst | EFindFirstObjectOptions::EnsureIfAmbiguous))
+		if (UObject* FoundObject = FindFirstObject<UObject>(*InObjectPath, EFindFirstObjectOptions::NativeFirst))
 		{
 			return FoundObject;
 		}
