Sequencer: Cinematic Viewport is now tied to the view range.

This allows the zooming in/out of the view to be in sync with Sequencer as opposed to the working range which is the outer range which generally does not shrink unless you reset the ranges (ie. Home button)

diff --git a/Engine/Plugins/MovieScene/LevelSequenceEditor/Source/LevelSequenceEditor/Private/CinematicViewport/SCinematicLevelViewport.cpp b/Engine/Plugins/MovieScene/LevelSequenceEditor/Source/LevelSequenceEditor/Private/CinematicViewport/SCinematicLevelViewport.cpp
--- a/Engine/Plugins/MovieScene/LevelSequenceEditor/Source/LevelSequenceEditor/Private/CinematicViewport/SCinematicLevelViewport.cpp
+++ b/Engine/Plugins/MovieScene/LevelSequenceEditor/Source/LevelSequenceEditor/Private/CinematicViewport/SCinematicLevelViewport.cpp
@@ -630,7 +630,7 @@ void SCinematicLevelViewport::Setup(FLevelSequenceEditorToolkit& NewToolkit)
 
 		if (TimeRangeContainer.IsValid())
 		{
-			const bool bShowWorkingRange = true, bShowViewRange = false, bShowPlaybackRange = true;
+			const bool bShowWorkingRange = false, bShowViewRange = true, bShowPlaybackRange = true;
 			TimeRangeContainer->SetContent(Sequencer->MakeTimeRange(DecoratedTransportControls.ToSharedRef(), bShowWorkingRange, bShowViewRange, bShowPlaybackRange));
 		}
 	}
diff --git a/Engine/Plugins/MovieScene/LevelSequenceEditor/Source/LevelSequenceEditor/Private/CinematicViewport/SCinematicTransportRange.cpp b/Engine/Plugins/MovieScene/LevelSequenceEditor/Source/LevelSequenceEditor/Private/CinematicViewport/SCinematicTransportRange.cpp
--- a/Engine/Plugins/MovieScene/LevelSequenceEditor/Source/LevelSequenceEditor/Private/CinematicViewport/SCinematicTransportRange.cpp
+++ b/Engine/Plugins/MovieScene/LevelSequenceEditor/Source/LevelSequenceEditor/Private/CinematicViewport/SCinematicTransportRange.cpp
@@ -99,10 +99,10 @@ void SCinematicTransportRange::Tick(const FGeometry& AllottedGeometry, const dou
 	ISequencer* Sequencer = GetSequencer();
 	if (Sequencer)
 	{
-		TRange<double> WorkingRange = Sequencer->GetFocusedMovieSceneSequence()->GetMovieScene()->GetEditorData().GetWorkingRange();
+		TRange<double> ViewRange = Sequencer->GetFocusedMovieSceneSequence()->GetMovieScene()->GetEditorData().GetViewRange();
 
 		// Anything within 3 pixel's worth of time is a duplicate as far as we're concerned
-		FTimeToPixel TimeToPixelConverter(AllottedGeometry, WorkingRange, Sequencer->GetFocusedTickResolution());
+		FTimeToPixel TimeToPixelConverter(AllottedGeometry, ViewRange, Sequencer->GetFocusedTickResolution());
 		const float DuplicateThreshold = (TimeToPixelConverter.PixelToSeconds(3.f) - TimeToPixelConverter.PixelToSeconds(0.f));
 
 		Sequencer->GetKeysFromSelection(ActiveKeyCollection, FMath::Max(DuplicateThreshold, SMALL_NUMBER));
@@ -123,11 +123,11 @@ void SCinematicTransportRange::DrawKeys(const FGeometry& AllottedGeometry, FSlat
 	const float TrackHeight = AllottedGeometry.GetLocalSize().Y - TrackOffsetY;
 
 	FFrameRate           TickResolution  = Sequencer->GetFocusedTickResolution();
-	TRange<double>       WorkingRange    = Sequencer->GetFocusedMovieSceneSequence()->GetMovieScene()->GetEditorData().GetWorkingRange();
+	TRange<double>       ViewRange       = Sequencer->GetFocusedMovieSceneSequence()->GetMovieScene()->GetEditorData().GetViewRange();
 	TRange<FFrameNumber> PlaybackRange   = Sequencer->GetFocusedMovieSceneSequence()->GetMovieScene()->GetPlaybackRange();
 
-	const FFrameTime FramesPerPixel = (WorkingRange.Size<double>() / AllottedGeometry.GetLocalSize().X) * TickResolution;
-	const float FullRange = WorkingRange.Size<float>();
+	const FFrameTime FramesPerPixel = (ViewRange.Size<double>() / AllottedGeometry.GetLocalSize().X) * TickResolution;
+	const float FullRange = ViewRange.Size<float>();
 
 	static const float BrushWidth = 7.f;
 	static const float BrushHeight = 7.f;
@@ -146,7 +146,7 @@ void SCinematicTransportRange::DrawKeys(const FGeometry& AllottedGeometry, FSlat
 			bOutPlayMarkerOnKey = true;
 		}
 
-		float Lerp = (Time/TickResolution - WorkingRange.GetLowerBoundValue()) / FullRange;
+		float Lerp = (Time/TickResolution - ViewRange.GetLowerBoundValue()) / FullRange;
 
 		FSlateDrawElement::MakeBox(
 			OutDrawElements,
@@ -180,14 +180,14 @@ int32 SCinematicTransportRange::OnPaint(const FPaintArgs& Args, const FGeometry&
 	const float TrackHeight = AllottedGeometry.GetLocalSize().Y - TrackOffsetY;
 
 	FFrameRate           TickResolution  = Sequencer->GetFocusedTickResolution();
-	TRange<double>       WorkingRange    = Sequencer->GetFocusedMovieSceneSequence()->GetMovieScene()->GetEditorData().GetWorkingRange();
+	TRange<double>       ViewRange       = Sequencer->GetFocusedMovieSceneSequence()->GetMovieScene()->GetEditorData().GetViewRange();
 	TRange<FFrameNumber> PlaybackRange   = Sequencer->GetFocusedMovieSceneSequence()->GetMovieScene()->GetPlaybackRange();
 
 	// Anything within 3 pixel's worth of time is a duplicate as far as we're concerned
-	FTimeToPixel TimeToPixelConverter(AllottedGeometry, WorkingRange, TickResolution);
+	FTimeToPixel TimeToPixelConverter(AllottedGeometry, ViewRange, TickResolution);
 	const float DuplicateThreshold = (TimeToPixelConverter.PixelToSeconds(3.f) - TimeToPixelConverter.PixelToSeconds(0.f));
 
-	const FFrameTime FramesPerPixel = (WorkingRange.Size<double>() / AllottedGeometry.GetLocalSize().X) * TickResolution;
+	const FFrameTime FramesPerPixel = (ViewRange.Size<double>() / AllottedGeometry.GetLocalSize().X) * TickResolution;
 
 	FColor DarkGray(40, 40, 40);
 	FColor MidGray(80, 80, 80);
@@ -203,10 +203,10 @@ int32 SCinematicTransportRange::OnPaint(const FPaintArgs& Args, const FGeometry&
 		FLinearColor(DarkGray)
 	);
 	
-	const float FullRange = WorkingRange.Size<float>();
+	const float FullRange = ViewRange.Size<float>();
 
-	const float PlaybackStartLerp	= (PlaybackRange.GetLowerBoundValue()/TickResolution - WorkingRange.GetLowerBoundValue()) / FullRange;
-	const float PlaybackEndLerp		= (PlaybackRange.GetUpperBoundValue()/TickResolution - WorkingRange.GetLowerBoundValue()) / FullRange;
+	const float PlaybackStartLerp	= (PlaybackRange.GetLowerBoundValue()/TickResolution - ViewRange.GetLowerBoundValue()) / FullRange;
+	const float PlaybackEndLerp		= (PlaybackRange.GetUpperBoundValue()/TickResolution - ViewRange.GetLowerBoundValue()) / FullRange;
 
 	// Draw the playback range
 	FSlateDrawElement::MakeBox(
@@ -219,7 +219,7 @@ int32 SCinematicTransportRange::OnPaint(const FPaintArgs& Args, const FGeometry&
 	);
 
 	const FQualifiedFrameTime CurrentTime = Sequencer->GetLocalTime();
-	const float ProgressLerp = (CurrentTime.AsSeconds() - WorkingRange.GetLowerBoundValue()) / FullRange;
+	const float ProgressLerp = (CurrentTime.AsSeconds() - ViewRange.GetLowerBoundValue()) / FullRange;
 
 	// Draw the playback progress
 	if (ProgressLerp > PlaybackStartLerp)
@@ -243,8 +243,8 @@ int32 SCinematicTransportRange::OnPaint(const FPaintArgs& Args, const FGeometry&
 	if (ActiveKeyCollection.IsValid())
 	{
 		TRange<FFrameNumber> VisibleFrameRange(
-			(WorkingRange.GetLowerBoundValue() * TickResolution).FloorToFrame(), 
-			(WorkingRange.GetUpperBoundValue() * TickResolution).CeilToFrame()
+			(ViewRange.GetLowerBoundValue() * TickResolution).FloorToFrame(),
+			(ViewRange.GetUpperBoundValue() * TickResolution).CeilToFrame()
 			);
 
 		TArrayView<const FFrameNumber> Keys = ActiveKeyCollection->GetKeysInRange(VisibleFrameRange);
