Sequencer: Restrict components based on the global class filter

Also added CineCameraComponent to the allowable component types.

diff --git a/Engine/Plugins/MovieScene/LevelSequenceEditor/Source/LevelSequenceEditor/LevelSequenceEditor.Build.cs b/Engine/Plugins/MovieScene/LevelSequenceEditor/Source/LevelSequenceEditor/LevelSequenceEditor.Build.cs
--- a/Engine/Plugins/MovieScene/LevelSequenceEditor/Source/LevelSequenceEditor/LevelSequenceEditor.Build.cs
+++ b/Engine/Plugins/MovieScene/LevelSequenceEditor/Source/LevelSequenceEditor/LevelSequenceEditor.Build.cs
@@ -26,6 +26,7 @@ public class LevelSequenceEditor : ModuleRules
                 "LevelSequence",
 				"BlueprintGraph",
                 "CinematicCamera",
+				"ClassViewer",
 				"Core",
 				"CoreUObject",
 				"CurveEditor",
diff --git a/Engine/Plugins/MovieScene/LevelSequenceEditor/Source/LevelSequenceEditor/Private/LevelSequenceEditorToolkit.cpp b/Engine/Plugins/MovieScene/LevelSequenceEditor/Source/LevelSequenceEditor/Private/LevelSequenceEditorToolkit.cpp
--- a/Engine/Plugins/MovieScene/LevelSequenceEditor/Source/LevelSequenceEditor/Private/LevelSequenceEditorToolkit.cpp
+++ b/Engine/Plugins/MovieScene/LevelSequenceEditor/Source/LevelSequenceEditor/Private/LevelSequenceEditorToolkit.cpp
@@ -53,6 +53,8 @@
 #include "MovieScene.h"
 #include "UnrealEdMisc.h"
 #include "ToolMenus.h"
+#include "ClassViewerFilter.h"
+#include "ClassViewerModule.h"
 
 // @todo sequencer: hack: setting defaults for transform tracks
 
@@ -880,6 +882,11 @@ void FLevelSequenceEditorToolkit::HandleTrackMenuExtensionAddTrack(FMenuBuilder&
 	{
 		return;
 	}
+	
+	FClassViewerModule& ClassViewerModule = FModuleManager::LoadModuleChecked<FClassViewerModule>("ClassViewer");
+	const TSharedPtr<IClassViewerFilter>& GlobalClassFilter = ClassViewerModule.GetGlobalClassViewerFilter();
+	TSharedRef<FClassViewerFilterFuncs> ClassFilterFuncs = ClassViewerModule.CreateFilterFuncs();
+	FClassViewerInitializationOptions ClassViewerOptions = {};
 
 	AActor* Actor = Cast<AActor>(ContextObjects[0]);
 	if (Actor != nullptr)
@@ -889,7 +896,7 @@ void FLevelSequenceEditorToolkit::HandleTrackMenuExtensionAddTrack(FMenuBuilder&
 			TMap<FString, UActorComponent*> SortedComponents;
 			for (UActorComponent* Component : Actor->GetComponents())
 			{
-				if (Component)
+				if (Component && !Component->IsVisualizationComponent() && GlobalClassFilter.IsValid() && GlobalClassFilter->IsClassAllowed(ClassViewerOptions, Component->GetClass(), ClassFilterFuncs))
 				{
 					SortedComponents.Add(Component->GetName(), Component);
 				}
