Sequencer: Follow up to Cinematic Viewport is now tied to the view range.

This allows the zooming in/out of the view to be in sync with Sequencer as opposed to the working range which is the outer range which generally does not shrink unless you reset the ranges (ie. Home button)

diff --git a/Engine/Plugins/MovieScene/LevelSequenceEditor/Source/LevelSequenceEditor/Private/CinematicViewport/SCinematicTransportRange.cpp b/Engine/Plugins/MovieScene/LevelSequenceEditor/Source/LevelSequenceEditor/Private/CinematicViewport/SCinematicTransportRange.cpp
--- a/Engine/Plugins/MovieScene/LevelSequenceEditor/Source/LevelSequenceEditor/Private/CinematicViewport/SCinematicTransportRange.cpp
+++ b/Engine/Plugins/MovieScene/LevelSequenceEditor/Source/LevelSequenceEditor/Private/CinematicViewport/SCinematicTransportRange.cpp
@@ -81,7 +81,7 @@ void SCinematicTransportRange::SetTime(const FGeometry& MyGeometry, const FPoint
 		Lerp = FMath::Clamp(Lerp, 0.f, 1.f);
 
 		FMovieSceneEditorData& EditorData = Sequencer->GetFocusedMovieSceneSequence()->GetMovieScene()->GetEditorData();
-		double NewTimeSeconds = EditorData.WorkStart + (EditorData.WorkEnd - EditorData.WorkStart) * Lerp;
+		double NewTimeSeconds = EditorData.ViewStart + (EditorData.ViewEnd - EditorData.ViewStart) * Lerp;
 
 		Sequencer->SetLocalTime(NewTimeSeconds * Sequencer->GetFocusedTickResolution(), ESnapTimeMode::STM_All);
 	}
