Sequencer: Follow up to restrict components based on the global class filter

diff --git a/Engine/Plugins/MovieScene/LevelSequenceEditor/Source/LevelSequenceEditor/Private/LevelSequenceEditorToolkit.cpp b/Engine/Plugins/MovieScene/LevelSequenceEditor/Source/LevelSequenceEditor/Private/LevelSequenceEditorToolkit.cpp
--- a/Engine/Plugins/MovieScene/LevelSequenceEditor/Source/LevelSequenceEditor/Private/LevelSequenceEditorToolkit.cpp
+++ b/Engine/Plugins/MovieScene/LevelSequenceEditor/Source/LevelSequenceEditor/Private/LevelSequenceEditorToolkit.cpp
@@ -896,7 +896,14 @@ void FLevelSequenceEditorToolkit::HandleTrackMenuExtensionAddTrack(FMenuBuilder&
 			TMap<FString, UActorComponent*> SortedComponents;
 			for (UActorComponent* Component : Actor->GetComponents())
 			{
-				if (Component && !Component->IsVisualizationComponent() && GlobalClassFilter.IsValid() && GlobalClassFilter->IsClassAllowed(ClassViewerOptions, Component->GetClass(), ClassFilterFuncs))
+				bool bValidComponent = Component && !Component->IsVisualizationComponent();
+
+				if (GlobalClassFilter.IsValid())
+				{
+					bValidComponent = GlobalClassFilter->IsClassAllowed(ClassViewerOptions, Component->GetClass(), ClassFilterFuncs);
+				}
+
+				if (bValidComponent)
 				{
 					SortedComponents.Add(Component->GetName(), Component);
 				}
