Sequencer: Follow up to restrict components based on the global class filter

diff --git a/Engine/Plugins/MovieScene/LevelSequenceEditor/Source/LevelSequenceEditor/Private/LevelSequenceEditorToolkit.cpp b/Engine/Plugins/MovieScene/LevelSequenceEditor/Source/LevelSequenceEditor/Private/LevelSequenceEditorToolkit.cpp
--- a/Engine/Plugins/MovieScene/LevelSequenceEditor/Source/LevelSequenceEditor/Private/LevelSequenceEditorToolkit.cpp
+++ b/Engine/Plugins/MovieScene/LevelSequenceEditor/Source/LevelSequenceEditor/Private/LevelSequenceEditorToolkit.cpp
@@ -896,16 +896,19 @@ void FLevelSequenceEditorToolkit::HandleTrackMenuExtensionAddTrack(FMenuBuilder&
 			TMap<FString, UActorComponent*> SortedComponents;
 			for (UActorComponent* Component : Actor->GetComponents())
 			{
-				bool bValidComponent = Component && !Component->IsVisualizationComponent();
-
-				if (GlobalClassFilter.IsValid())
+				if (Component)
 				{
-					bValidComponent = GlobalClassFilter->IsClassAllowed(ClassViewerOptions, Component->GetClass(), ClassFilterFuncs);
-				}
+					bool bValidComponent = Component && !Component->IsVisualizationComponent();
 
-				if (bValidComponent)
-				{
-					SortedComponents.Add(Component->GetName(), Component);
+					if (GlobalClassFilter.IsValid())
+					{
+						bValidComponent = GlobalClassFilter->IsClassAllowed(ClassViewerOptions, Component->GetClass(), ClassFilterFuncs);
+					}
+
+					if (bValidComponent)
+					{
+						SortedComponents.Add(Component->GetName(), Component);
+					}
 				}
 			}
 			SortedComponents.KeySort([](const FString& A, const FString& B) 
