Sequencer: Fix bake which was computing the global time correctly but then calling SetLocalTimeDirectly instead of SetGlobalTime

diff --git a/Engine/Plugins/MovieScene/LevelSequenceEditor/Source/LevelSequenceEditor/Private/LevelSequenceEditorSubsystem.cpp b/Engine/Plugins/MovieScene/LevelSequenceEditor/Source/LevelSequenceEditor/Private/LevelSequenceEditorSubsystem.cpp
--- a/Engine/Plugins/MovieScene/LevelSequenceEditor/Source/LevelSequenceEditor/Private/LevelSequenceEditorSubsystem.cpp
+++ b/Engine/Plugins/MovieScene/LevelSequenceEditor/Source/LevelSequenceEditor/Private/LevelSequenceEditorSubsystem.cpp
@@ -730,7 +730,7 @@ void ULevelSequenceEditorSubsystem::BakeTransform(const TArray<FMovieSceneBindin
 		FFrameNumber KeyTime = FFrameRate::Snap(EvalTime, TickResolution, DisplayRate).FloorToFrame();
 		FMovieSceneEvaluationRange Range(KeyTime * RootToLocalTransform.InverseLinearOnly(), TickResolution);
 
-		Sequencer->SetLocalTimeDirectly(Range.GetTime());
+		Sequencer->SetGlobalTime(Range.GetTime());
 
 		for (const FMovieSceneBindingProxy& ObjectBinding : ObjectBindings)
 		{
