LevelSequenceEditor: List transient actors when binding onto transient LevelSequences.

diff --git a/Engine/Plugins/MovieScene/LevelSequenceEditor/Source/LevelSequenceEditor/Private/Misc/LevelSequenceEditorActorBinding.cpp b/Engine/Plugins/MovieScene/LevelSequenceEditor/Source/LevelSequenceEditor/Private/Misc/LevelSequenceEditorActorBinding.cpp
--- a/Engine/Plugins/MovieScene/LevelSequenceEditor/Source/LevelSequenceEditor/Private/Misc/LevelSequenceEditorActorBinding.cpp
+++ b/Engine/Plugins/MovieScene/LevelSequenceEditor/Source/LevelSequenceEditor/Private/Misc/LevelSequenceEditorActorBinding.cpp
@@ -48,9 +48,10 @@ void FLevelSequenceEditorActorBinding::AddPossessActorMenuExtensions(FMenuBuilde
 	// pointer comparison and not an object comparison, and the async list shouldn't run the filter if the object is no longer valid.
 	// We don't need to check against Sequencer spawnables as they're not valid for possession.
 	TSet<UObject*> ExistingPossessedObjects;
+	UMovieSceneSequence* MovieSceneSequence = nullptr;
 	if (Sequencer.IsValid())
 	{
-		UMovieSceneSequence* MovieSceneSequence = Sequencer.Pin()->GetFocusedMovieSceneSequence();
+		MovieSceneSequence = Sequencer.Pin()->GetFocusedMovieSceneSequence();
 		UMovieScene* MovieScene = MovieSceneSequence->GetMovieScene();
 		if(MovieScene)
 		{
@@ -97,7 +98,7 @@ void FLevelSequenceEditorActorBinding::AddPossessActorMenuExtensions(FMenuBuilde
 			AddActorsToSequencer(ActorsValidForPossession.GetData(), ActorsValidForPossession.Num());
 		}));
 	}
-	
+
 	MenuBuilder.BeginSection("ChooseActorSection", LOCTEXT("ChooseActor", "Choose Actor:"));
 
 	// Set up a menu entry to add any arbitrary actor to the sequencer
@@ -109,6 +110,9 @@ void FLevelSequenceEditorActorBinding::AddPossessActorMenuExtensions(FMenuBuilde
 		InitOptions.bShowCreateNewFolder = false;
 		InitOptions.bFocusSearchBoxWhenOpened = true;
 
+		// Allow transient actors if the level sequence itself is transient (the expectation is that these would never be saved)
+		InitOptions.bShowTransient = MovieSceneSequence && MovieSceneSequence->GetOutermost() == GetTransientPackage();
+
 		// Only want the actor label column
 		InitOptions.ColumnMap.Add(FSceneOutlinerBuiltInColumnTypes::Label(), FSceneOutlinerColumnInfo(ESceneOutlinerColumnVisibility::Visible, 0));
 
