Sequencer: Convert to spawnable now creates a transform track that takes the transform origin into account

diff --git a/Engine/Plugins/MovieScene/LevelSequenceEditor/Source/LevelSequenceEditor/Private/Misc/LevelSequenceEditorActorSpawner.cpp b/Engine/Plugins/MovieScene/LevelSequenceEditor/Source/LevelSequenceEditor/Private/Misc/LevelSequenceEditorActorSpawner.cpp
--- a/Engine/Plugins/MovieScene/LevelSequenceEditor/Source/LevelSequenceEditor/Private/Misc/LevelSequenceEditorActorSpawner.cpp
+++ b/Engine/Plugins/MovieScene/LevelSequenceEditor/Source/LevelSequenceEditor/Private/Misc/LevelSequenceEditorActorSpawner.cpp
@@ -8,7 +8,9 @@
 #include "SnappingUtils.h"
 #include "Sections/MovieScene3DTransformSection.h"
 #include "Tracks/MovieScene3DTransformTrack.h"
+#include "Tracks/IMovieSceneTransformOrigin.h"
 #include "SequencerSettings.h"
+#include "IMovieScenePlaybackClient.h"
 #include "ISequencer.h"
 #include "Tracks/MovieSceneSpawnTrack.h"
 #include "Sections/MovieSceneBoolSection.h"
@@ -262,9 +264,29 @@ void FLevelSequenceEditorActorSpawner::SetupDefaultsForSpawnable(UObject* Spawne
 			UMovieScene3DTransformSection* TransformSection = Cast<UMovieScene3DTransformSection>(TransformTrack->CreateNewSection());
 			TransformTrack->AddSection(*TransformSection);
 
-			FVector Location = DefaultTransform->Translation;
-			FVector Rotation = DefaultTransform->Rotation.Euler();
-			FVector Scale    = DefaultTransform->Scale;
+			FTransform TransformOrigin;
+
+			const IMovieScenePlaybackClient* Client = Sequencer->GetPlaybackClient();
+			const UObject* InstanceData = Client ? Client->GetInstanceData() : nullptr;
+			const IMovieSceneTransformOrigin* RawInterface = Cast<const IMovieSceneTransformOrigin>(InstanceData);
+
+			const bool bHasInterface = RawInterface || (InstanceData && InstanceData->GetClass()->ImplementsInterface(UMovieSceneTransformOrigin::StaticClass()));
+			if (bHasInterface)
+			{
+				// Retrieve the current origin
+				TransformOrigin = RawInterface ? RawInterface->GetTransformOrigin() : IMovieSceneTransformOrigin::Execute_BP_GetTransformOrigin(InstanceData);
+			}
+
+			FTransform Transform;
+			Transform.SetLocation(DefaultTransform->Translation);
+			Transform.SetRotation(DefaultTransform->Rotation.Quaternion());
+			Transform.SetScale3D(DefaultTransform->Scale);
+
+			Transform = Transform * TransformOrigin.Inverse();
+
+			FVector Location = Transform.GetLocation();
+			FVector Rotation = Transform.GetRotation().Euler();
+			FVector Scale    = Transform.GetScale3D();
 
 			// Set the section to be infinite if necessary
 			if (Sequencer->GetInfiniteKeyAreas())
