Sequencer: Added the ability to bind to a skeletal mesh component's post process instance.

diff --git a/Engine/Plugins/MovieScene/LevelSequenceEditor/Source/LevelSequenceEditor/Private/LevelSequenceEditorToolkit.cpp b/Engine/Plugins/MovieScene/LevelSequenceEditor/Source/LevelSequenceEditor/Private/LevelSequenceEditorToolkit.cpp
--- a/Engine/Plugins/MovieScene/LevelSequenceEditor/Source/LevelSequenceEditor/Private/LevelSequenceEditorToolkit.cpp
+++ b/Engine/Plugins/MovieScene/LevelSequenceEditor/Source/LevelSequenceEditor/Private/LevelSequenceEditorToolkit.cpp
@@ -933,30 +933,49 @@ void FLevelSequenceEditorToolkit::HandleTrackMenuExtensionAddTrack(FMenuBuilder&
 			UAnimInstance* AnimInstance = SkeletalComponent->GetAnimInstance();
 			
 			FText AnimInstanceLabel = LOCTEXT("AnimInstanceLabel", "Anim Instance");
-			FText DetailedInstanceText = AnimInstance
-				? FText::Format(LOCTEXT("AnimInstanceLabelFormat", "Anim Instance '{0}'"), FText::FromString(AnimInstance->GetName()))
+			FText DetailedAnimInstanceText = AnimInstance ? 
+				FText::Format(LOCTEXT("AnimInstanceLabelFormat", "Anim Instance '{0}'"), FText::FromString(AnimInstance->GetName()))
 				: AnimInstanceLabel;
 
 			AddTrackMenuBuilder.BeginSection("Anim Instance", AnimInstanceLabel);
 			{
 				AddTrackMenuBuilder.AddMenuEntry(
-					DetailedInstanceText,
+					DetailedAnimInstanceText,
 					LOCTEXT("AnimInstanceToolTip", "Add this skeletal mesh component's animation instance."),
 					FSlateIcon(),
 					FUIAction(
-						FExecuteAction::CreateSP(this, &FLevelSequenceEditorToolkit::BindAnimationInstance, SkeletalComponent)
+						FExecuteAction::CreateSP(this, &FLevelSequenceEditorToolkit::BindAnimationInstance, SkeletalComponent, AnimInstance)
 					)
 				);
 			}
 			AddTrackMenuBuilder.EndSection();
+
+			if (UAnimInstance* PostProcessInstance = SkeletalComponent->GetPostProcessInstance())
+			{
+				FText PostProcessInstanceLabel = LOCTEXT("PostProcessInstanceLabel", "Post Process Instance");
+				FText DetailedPostProcessInstanceText = PostProcessInstance ?
+					FText::Format(LOCTEXT("PostProcessInstanceLabelFormat", "Post Process Instance '{0}'"), FText::FromString(PostProcessInstance->GetName()))
+					: PostProcessInstanceLabel;
+
+				AddTrackMenuBuilder.BeginSection("Post Process Instance", PostProcessInstanceLabel);
+				{
+					AddTrackMenuBuilder.AddMenuEntry(
+						DetailedPostProcessInstanceText,
+						LOCTEXT("PostProcessInstanceToolTip", "Add this skeletal mesh component's post process instance."),
+						FSlateIcon(),
+						FUIAction(
+							FExecuteAction::CreateSP(this, &FLevelSequenceEditorToolkit::BindAnimationInstance, SkeletalComponent, PostProcessInstance)
+						)
+					);
+				}
+				AddTrackMenuBuilder.EndSection();
+			}
 		}
 	}
 }
 
-void FLevelSequenceEditorToolkit::BindAnimationInstance(USkeletalMeshComponent* SkeletalComponent)
+void FLevelSequenceEditorToolkit::BindAnimationInstance(USkeletalMeshComponent* SkeletalComponent, UAnimInstance* AnimInstance)
 {
-	UAnimInstance* AnimInstance = SkeletalComponent->GetAnimInstance();
-
 	// If there is no script instance at the moment, just use a dummy instance for the purposes of setting up the binding in the first place.
 	// This temporary object will get GC'd later on and is never actually applied to the anim instance
 	Sequencer->GetHandleToObject(AnimInstance ? AnimInstance : NewObject<UAnimInstance>(SkeletalComponent));
diff --git a/Engine/Plugins/MovieScene/LevelSequenceEditor/Source/LevelSequenceEditor/Private/LevelSequenceEditorToolkit.h b/Engine/Plugins/MovieScene/LevelSequenceEditor/Source/LevelSequenceEditor/Private/LevelSequenceEditorToolkit.h
--- a/Engine/Plugins/MovieScene/LevelSequenceEditor/Source/LevelSequenceEditor/Private/LevelSequenceEditorToolkit.h
+++ b/Engine/Plugins/MovieScene/LevelSequenceEditor/Source/LevelSequenceEditor/Private/LevelSequenceEditorToolkit.h
@@ -20,6 +20,7 @@ class FToolBarBuilder;
 class IAssetViewport;
 class ISequencer;
 class UActorComponent;
+class UAnimInstance;
 class ULevelSequence;
 class UMovieSceneCinematicShotTrack;
 class FLevelSequencePlaybackContext;
@@ -128,7 +129,7 @@ private:
 	void HandleAddComponentActionExecute(UActorComponent* Component);
 
 	/** Create a new binding for the specified skeletal mesh component's animation instance. */
-	void BindAnimationInstance(USkeletalMeshComponent* SkeletalComponent);
+	void BindAnimationInstance(USkeletalMeshComponent* SkeletalComponent, UAnimInstance* AnimInstance);
 
 	/** Callback for map changes. */
 	void HandleMapChanged(UWorld* NewWorld, EMapChangeType MapChangeType);
diff --git a/Engine/Source/Runtime/LevelSequence/Private/LevelSequenceBindingReference.cpp b/Engine/Source/Runtime/LevelSequence/Private/LevelSequenceBindingReference.cpp
--- a/Engine/Source/Runtime/LevelSequence/Private/LevelSequenceBindingReference.cpp
+++ b/Engine/Source/Runtime/LevelSequence/Private/LevelSequenceBindingReference.cpp
@@ -218,14 +218,25 @@ bool FLevelSequenceObjectReferenceMap::Serialize(FArchive& Ar)
 
 bool FLevelSequenceBindingReferences::HasBinding(const FGuid& ObjectId) const
 {
-	return BindingIdToReferences.Contains(ObjectId) || AnimSequenceInstances.Contains(ObjectId);
+	return BindingIdToReferences.Contains(ObjectId) || AnimSequenceInstances.Contains(ObjectId) || PostProcessInstances.Contains(ObjectId);
 }
 
 void FLevelSequenceBindingReferences::AddBinding(const FGuid& ObjectId, UObject* InObject, UObject* InContext)
 {
-	if (InObject->IsA<UAnimInstance>())
+	if (UAnimInstance* AnimInstance = Cast<UAnimInstance>(InObject))
 	{
-		AnimSequenceInstances.Add(ObjectId);
+		if (AnimInstance->GetOwningComponent()->GetAnimInstance() == InObject)
+		{
+			AnimSequenceInstances.Add(ObjectId);
+		}
+		else if (AnimInstance->GetOwningComponent()->GetPostProcessInstance() == InObject)
+		{
+			PostProcessInstances.Add(ObjectId);
+		}
+		else
+		{
+			UE_LOG(LogMovieScene, Warning, TEXT("Attempted to add a binding for %s which is not an anim instance or post process instance"), *GetNameSafe(InObject));
+		}
 	}
 	else
 	{
@@ -301,9 +312,16 @@ void FLevelSequenceBindingReferences::ResolveBinding(const FGuid& ObjectId, UObj
 	else if (USkeletalMeshComponent* SkeletalMeshComponent = Cast<USkeletalMeshComponent>(InContext))
 	{
 		// If the object ID exists in the AnimSequenceInstances set, then this binding relates to an anim instance on a skeletal mesh component
-		if (SkeletalMeshComponent && AnimSequenceInstances.Contains(ObjectId) && SkeletalMeshComponent->GetAnimInstance())
+		if (SkeletalMeshComponent)
 		{
-			OutObjects.Add(SkeletalMeshComponent->GetAnimInstance());
+			if (AnimSequenceInstances.Contains(ObjectId) && SkeletalMeshComponent->GetAnimInstance())
+			{
+				OutObjects.Add(SkeletalMeshComponent->GetAnimInstance());
+			}
+			else if (PostProcessInstances.Contains(ObjectId) && SkeletalMeshComponent->GetPostProcessInstance())
+			{
+				OutObjects.Add(SkeletalMeshComponent->GetPostProcessInstance());
+			}
 		}
 	}
 }
diff --git a/Engine/Source/Runtime/LevelSequence/Public/LevelSequenceBindingReference.h b/Engine/Source/Runtime/LevelSequence/Public/LevelSequenceBindingReference.h
--- a/Engine/Source/Runtime/LevelSequence/Public/LevelSequenceBindingReference.h
+++ b/Engine/Source/Runtime/LevelSequence/Public/LevelSequenceBindingReference.h
@@ -170,6 +170,14 @@ struct FLevelSequenceBindingReferences
 		return AnimSequenceInstances;
 	}
 
+	/**
+	 * Const accessor for the currently bound post process instance IDs
+	 */
+	const TSet<FGuid>& GetBoundPostProcessInstances() const
+	{
+		return PostProcessInstances;
+	}
+
 	/**
 	 * Filter out any bindings that do not match the specified set of GUIDs
 	 *
@@ -186,4 +194,8 @@ private:
 	/** A set of object binding IDs that relate to anim sequence instances (must be a child of USkeletalMeshComponent) */
 	UPROPERTY()
 	TSet<FGuid> AnimSequenceInstances;
+
+	/** A set of object binding IDs that relate to post process instances (must be a child of USkeletalMeshComponent) */
+	UPROPERTY()
+	TSet<FGuid> PostProcessInstances;
 };
