Sequencer: Remove Master from Sequencer

- Remove all instances of Master from the UI - ie. Replace Create Master Sequence with Create Sequence with Shots
- Internally, remove Master - Master Tracks in the MovieScene can just be Tracks

Deprecations (Properties):
- UMovieScene:MasterTracks deprecated for Tracks
- UMovieSceneFolder:ChildMasterTracks deprecated for ChildTracks
- FLevelSequencePlayerSnapshot: MasterName, MasterTime for RootName, RootTime
- UMoviePipelineSetting_BlueprintBase: bIsValidOnMaster for bIsValidOnPrimary

Deprecations (Functions):
- UMovieScene: GetMasterTracks, FindMasterTracksByType, FindMasterTracksByExactType, AddMasterTrack, RemoveMasterTrack deprecated for GetTracks, FindTracksByType, FindTracksByExactType, AddTrack, RemoveTrack
- UMovieSceneFolder: GetChildMasterTracks, AddChildMasterTrack, RemoveChildMasterTrack deprecated for GetChildTracks, AddChildTracks, RemoveChildTrack
- FMovieSceneTrackEditor: FFindOrCreateMasterTrackResult, FindOrCreateMasterTrack deprecated for FindOrCreateRootTrackResult, FindOrCreateRootTrack
- ULevelSequenceDirector: GetMasterSequenceTime deprecated for ULevelSequenceDirector::GetRootSequenceTime
- UDaySequenceDirector: GetMasterSequenceTime deprecated for UDaySequenceDirector::GetRootSequenceTime
- UMoviePipelineSetting: IsValidOnMaster for IsValidOnPrimary
- UMoviePipelineBlueprintLibrary: GetMasterTimecode, GetMasterFrameNumber for GetRootTimecode, GetRootFrameNumber
- UMoviePipeline: GetPipelineMasterConfig for GetPipelinePrimaryConfig
- UMovieSceneAudioTrack: IsAMasterTrack

Classes Renamed:
- FAnimTypePreAnimatedStateMasterStorage for FAnimTypePreAnimatedStateRootStorage
- UMovieSceneMasterInstantiatorSystem for UMovieSceneRootInstantiatorSystem
- ULevelSequenceMasterSequenceSettings to ULevelSequenceSequenceWithShotsSetting

Note, the following instances of master still occur and a "KnownCase" has been added for them in CheckAcceptableWords.cs:
- Deprecated properties (MasterTracks -> Tracks) and functions (ie. GetMasterTracks -> GetTracks)
- masterclip in FCPXML since it's an industry standard term
- masteraudiosubmix in AudioCaptureProtocol
- mastersubmix in MoviePipelineAudioRendering
- masteertimecode, masterframenumber, masterconfig, validonmaster in MovieRenderPipeline
- mastersequencetime, mastertime, mastername in LevelSequence

diff --git a/Engine/Config/BaseEngine.ini b/Engine/Config/BaseEngine.ini
--- a/Engine/Config/BaseEngine.ini
+++ b/Engine/Config/BaseEngine.ini
@@ -1099,6 +1099,8 @@ DefaultWorldPartitionHLODLayer=/Engine/Maps/Templates/HLODs/HLODLayer_Instanced.
 ; 5.2
 +PropertyRedirects=(OldName="DecalComponent.SetFadeIn.Duaration", NewName="DecalComponent.SetFadeIn.Duration")
 
++ClassRedirects=(OldName="/Script/MovieRenderPipelineCore.MoviePipelineMasterConfig", NewName="/Script/MovieRenderPipelineCore.MoviePipelinePrimaryConfig")
+
 [CoreUObject.Metadata]
 ; Note: UnrealHeaderTool should be rerun after making changes to MetadataRedirects to catch any keys specified in class headers
 +MetadataRedirects=(OldKey="K2Protected", NewKey="BlueprintProtected")
diff --git a/Engine/Plugins/Animation/AnimationData/Source/AnimationData/Private/AnimSequencerController.cpp b/Engine/Plugins/Animation/AnimationData/Source/AnimationData/Private/AnimSequencerController.cpp
--- a/Engine/Plugins/Animation/AnimationData/Source/AnimationData/Private/AnimSequencerController.cpp
+++ b/Engine/Plugins/Animation/AnimationData/Source/AnimationData/Private/AnimSequencerController.cpp
@@ -238,7 +238,7 @@ void UAnimSequencerController::EnsureModelIsInitialized() const
 {
 	ensure(Model.Get());
 	ensure(Model->MovieScene);
-	const UMovieSceneControlRigParameterTrack* Track = Model->MovieScene->FindMasterTrack<UMovieSceneControlRigParameterTrack>();
+	const UMovieSceneControlRigParameterTrack* Track = Model->MovieScene->FindTrack<UMovieSceneControlRigParameterTrack>();
 	if (ensure(Track))
 	{
 		ensure(Track->GetAllSections().Num() > 0);
@@ -2464,7 +2464,7 @@ void UAnimSequencerController::InitializeModel()
 			const TRange<FFrameNumber> DataRange = TRange<FFrameNumber>::Inclusive(FFrameNumber(0), 1);
 			MovieScene->SetPlaybackRange(DataRange);
 			
-			if(UMovieSceneControlRigParameterTrack* Track = MovieScene->AddMasterTrack<UMovieSceneControlRigParameterTrack>())
+			if(UMovieSceneControlRigParameterTrack* Track = MovieScene->AddTrack<UMovieSceneControlRigParameterTrack>())
 			{
 				if(UFKControlRig* ControlRig = NewObject<UFKControlRig>(Track, UFKControlRig::StaticClass(), NAME_None, RF_Transactional))
 				{
diff --git a/Engine/Plugins/Animation/AnimationData/Source/AnimationData/Private/AnimSequencerDataModel.cpp b/Engine/Plugins/Animation/AnimationData/Source/AnimationData/Private/AnimSequencerDataModel.cpp
--- a/Engine/Plugins/Animation/AnimationData/Source/AnimationData/Private/AnimSequencerDataModel.cpp
+++ b/Engine/Plugins/Animation/AnimationData/Source/AnimationData/Private/AnimSequencerDataModel.cpp
@@ -705,7 +705,7 @@ void UAnimationSequencerDataModel::OnNotify(const EAnimDataModelNotifyType& Noti
 
 UMovieSceneControlRigParameterTrack* UAnimationSequencerDataModel::GetControlRigTrack() const
 {
-	return MovieScene->FindMasterTrack<UMovieSceneControlRigParameterTrack>();
+	return MovieScene->FindTrack<UMovieSceneControlRigParameterTrack>();
 }
 
 UMovieSceneControlRigParameterSection* UAnimationSequencerDataModel::GetFKControlRigSection() const
@@ -916,10 +916,10 @@ void UAnimationSequencerDataModel::ValidateSequencerData() const
 {
 	checkf(MovieScene, TEXT("No Movie Scene found for SequencerDataModel"));
 
-	const int32 NumberOfMasterTracks = MovieScene->GetMasterTracks().Num();
-	checkf(NumberOfMasterTracks == 1, TEXT("Invalid number of Tracks in Movie Scene expected 1 but found %i"), NumberOfMasterTracks);
+	const int32 NumberOfTracks = MovieScene->GetTracks().Num();
+	checkf(NumberOfTracks == 1, TEXT("Invalid number of Tracks in Movie Scene expected 1 but found %i"), NumberOfTracks);
 		
-	const UMovieSceneControlRigParameterTrack* Track = MovieScene->FindMasterTrack<UMovieSceneControlRigParameterTrack>();
+	const UMovieSceneControlRigParameterTrack* Track = MovieScene->FindTrack<UMovieSceneControlRigParameterTrack>();
 	checkf(Track, TEXT("Unable to find Control Rig Track"));
 
 	const int32 NumberOfSections = Track->GetAllSections().Num();
diff --git a/Engine/Plugins/Animation/LiveLink/Source/LiveLinkSequencer/Private/Sequencer/LiveLinkPropertyTrackEditor.cpp b/Engine/Plugins/Animation/LiveLink/Source/LiveLinkSequencer/Private/Sequencer/LiveLinkPropertyTrackEditor.cpp
--- a/Engine/Plugins/Animation/LiveLink/Source/LiveLinkSequencer/Private/Sequencer/LiveLinkPropertyTrackEditor.cpp
+++ b/Engine/Plugins/Animation/LiveLink/Source/LiveLinkSequencer/Private/Sequencer/LiveLinkPropertyTrackEditor.cpp
@@ -244,12 +244,12 @@ void FLiveLinkPropertyTrackEditor::HandleAddLiveLinkTrackMenuEntryExecute()
 		return;
 	}
 
-	UMovieSceneTrack* Track = FocusedMovieScene->FindMasterTrack<UMovieSceneLiveLinkTrack>();
+	UMovieSceneTrack* Track = FocusedMovieScene->FindTrack<UMovieSceneLiveLinkTrack>();
 	if (!Track)
 	{
 		const FScopedTransaction Transaction(NSLOCTEXT("Sequencer", "AddLiveLinkTrack_Transaction", "Add Live Link Track"));
 		FocusedMovieScene->Modify();
-		UMovieSceneLiveLinkTrack* NewTrack = FocusedMovieScene->AddMasterTrack<UMovieSceneLiveLinkTrack>();
+		UMovieSceneLiveLinkTrack* NewTrack = FocusedMovieScene->AddTrack<UMovieSceneLiveLinkTrack>();
 		ensure(NewTrack);
 
 		NewTrack->SetDisplayName(LOCTEXT("LiveLinkTrackName", "Live Link"));
@@ -264,7 +264,7 @@ void FLiveLinkPropertyTrackEditor::HandleAddLiveLinkTrackMenuEntryExecute()
 bool FLiveLinkPropertyTrackEditor::HandleAddLiveLinkTrackMenuEntryCanExecute() const
 {
 	UMovieScene* FocusedMovieScene = GetFocusedMovieScene();
-	return ((FocusedMovieScene != nullptr) && (FocusedMovieScene->FindMasterTrack<UMovieSceneLiveLinkTrack>() == nullptr));
+	return ((FocusedMovieScene != nullptr) && (FocusedMovieScene->FindTrack<UMovieSceneLiveLinkTrack>() == nullptr));
 }
 
 
diff --git a/Engine/Plugins/Animation/LiveLink/Source/LiveLinkSequencer/Private/SequencerRecorderSections/MovieSceneLiveLinkTrackRecorder.cpp b/Engine/Plugins/Animation/LiveLink/Source/LiveLinkSequencer/Private/SequencerRecorderSections/MovieSceneLiveLinkTrackRecorder.cpp
--- a/Engine/Plugins/Animation/LiveLink/Source/LiveLinkSequencer/Private/SequencerRecorderSections/MovieSceneLiveLinkTrackRecorder.cpp
+++ b/Engine/Plugins/Animation/LiveLink/Source/LiveLinkSequencer/Private/SequencerRecorderSections/MovieSceneLiveLinkTrackRecorder.cpp
@@ -49,14 +49,14 @@ void UMovieSceneLiveLinkTrackRecorder::CreateTrack(UMovieScene* InMovieScene, co
 	CreateTracks();
 }
 
-UMovieSceneLiveLinkTrack* UMovieSceneLiveLinkTrackRecorder::DoesLiveLinkMasterTrackExist(const FName& MasterTrackName, const TSubclassOf<ULiveLinkRole>& InTrackRole)
+UMovieSceneLiveLinkTrack* UMovieSceneLiveLinkTrackRecorder::DoesLiveLinkTrackExist(const FName& TrackName, const TSubclassOf<ULiveLinkRole>& InTrackRole)
 {
-	for (UMovieSceneTrack* MasterTrack : MovieScene->GetMasterTracks())
+	for (UMovieSceneTrack* Track : MovieScene->GetTracks())
 	{
-		if (MasterTrack->IsA(UMovieSceneLiveLinkTrack::StaticClass()))
+		if (Track->IsA(UMovieSceneLiveLinkTrack::StaticClass()))
 		{
-			UMovieSceneLiveLinkTrack* TestLiveLinkTrack = CastChecked<UMovieSceneLiveLinkTrack>(MasterTrack);
-			if (TestLiveLinkTrack && TestLiveLinkTrack->GetPropertyName() == MasterTrackName && TestLiveLinkTrack->GetTrackRole() == InTrackRole)
+			UMovieSceneLiveLinkTrack* TestLiveLinkTrack = CastChecked<UMovieSceneLiveLinkTrack>(Track);
+			if (TestLiveLinkTrack && TestLiveLinkTrack->GetPropertyName() == TrackName && TestLiveLinkTrack->GetTrackRole() == InTrackRole)
 			{
 				return TestLiveLinkTrack;
 			}
@@ -114,10 +114,10 @@ void UMovieSceneLiveLinkTrackRecorder::CreateTracks()
 
 	if(bRegistered)
 	{
-		LiveLinkTrack = DoesLiveLinkMasterTrackExist(SubjectName, SubjectRole);
+		LiveLinkTrack = DoesLiveLinkTrackExist(SubjectName, SubjectRole);
 		if (!LiveLinkTrack.IsValid())
 		{
-			LiveLinkTrack = MovieScene->AddMasterTrack<UMovieSceneLiveLinkTrack>();
+			LiveLinkTrack = MovieScene->AddTrack<UMovieSceneLiveLinkTrack>();
 			LiveLinkTrack->SetTrackRole(SubjectRole);
 		}
 		else
@@ -309,7 +309,7 @@ void UMovieSceneLiveLinkTrackRecorder::AddContentsToFolder(UMovieSceneFolder* In
 {
 	if (LiveLinkTrack.IsValid())
 	{
-		InFolder->AddChildMasterTrack(LiveLinkTrack.Get());
+		InFolder->AddChildTrack(LiveLinkTrack.Get());
 	}
 }
 
diff --git a/Engine/Plugins/Animation/LiveLink/Source/LiveLinkSequencer/Private/TakeRecorderSource/TakeRecorderLiveLinkSource.cpp b/Engine/Plugins/Animation/LiveLink/Source/LiveLinkSequencer/Private/TakeRecorderSource/TakeRecorderLiveLinkSource.cpp
--- a/Engine/Plugins/Animation/LiveLink/Source/LiveLinkSequencer/Private/TakeRecorderSource/TakeRecorderLiveLinkSource.cpp
+++ b/Engine/Plugins/Animation/LiveLink/Source/LiveLinkSequencer/Private/TakeRecorderSource/TakeRecorderLiveLinkSource.cpp
@@ -31,7 +31,7 @@ UTakeRecorderLiveLinkSource::UTakeRecorderLiveLinkSource(const FObjectInitialize
 	TrackTint = FColor(74, 108, 164);
 }
 
-TArray<UTakeRecorderSource*> UTakeRecorderLiveLinkSource::PreRecording(ULevelSequence* InSequence, FMovieSceneSequenceID InSequenceID, ULevelSequence* InMasterSequence, FManifestSerializer* InManifestSerializer) 
+TArray<UTakeRecorderSource*> UTakeRecorderLiveLinkSource::PreRecording(ULevelSequence* InSequence, FMovieSceneSequenceID InSequenceID, ULevelSequence* InRootSequence, FManifestSerializer* InManifestSerializer) 
 {
 	UMovieScene* MovieScene = InSequence->GetMovieScene();
 	TrackRecorder = NewObject<UMovieSceneLiveLinkTrackRecorder>();
@@ -65,7 +65,7 @@ void UTakeRecorderLiveLinkSource::StopRecording(class ULevelSequence* InSequence
 	}
 }
 
-TArray<UTakeRecorderSource*> UTakeRecorderLiveLinkSource::PostRecording(class ULevelSequence* InSequence, class ULevelSequence* InMasterSequence, const bool bCancelled)
+TArray<UTakeRecorderSource*> UTakeRecorderLiveLinkSource::PostRecording(class ULevelSequence* InSequence, class ULevelSequence* InRootSequence, const bool bCancelled)
 {
 	if (TrackRecorder)
 	{
diff --git a/Engine/Plugins/Animation/LiveLink/Source/LiveLinkSequencer/Private/TakeRecorderSource/TakeRecorderLiveLinkSource.h b/Engine/Plugins/Animation/LiveLink/Source/LiveLinkSequencer/Private/TakeRecorderSource/TakeRecorderLiveLinkSource.h
--- a/Engine/Plugins/Animation/LiveLink/Source/LiveLinkSequencer/Private/TakeRecorderSource/TakeRecorderLiveLinkSource.h
+++ b/Engine/Plugins/Animation/LiveLink/Source/LiveLinkSequencer/Private/TakeRecorderSource/TakeRecorderLiveLinkSource.h
@@ -76,7 +76,7 @@ public:
 	bool bDiscardSamplesBeforeStart;
 
 	/**
-	* The master track recorder we created.
+	* The track recorder we created.
 	*/
 	UPROPERTY()
 	TObjectPtr<UMovieSceneLiveLinkTrackRecorder> TrackRecorder;
@@ -84,11 +84,11 @@ public:
 private:
 
 	// UTakeRecorderSource
-	virtual TArray<UTakeRecorderSource*> PreRecording(ULevelSequence* InSequence, FMovieSceneSequenceID InSequenceID, ULevelSequence* InMasterSequence, FManifestSerializer* InManifestSerializer) override;
+	virtual TArray<UTakeRecorderSource*> PreRecording(ULevelSequence* InSequence, FMovieSceneSequenceID InSequenceID, ULevelSequence* InRootSequence, FManifestSerializer* InManifestSerializer) override;
 	virtual void StartRecording(const FTimecode& InSectionStartTimecode, const FFrameNumber& InSectionFirstFrame, class ULevelSequence* InSequence) override;
 	virtual void TickRecording(const FQualifiedFrameTime& CurrentSequenceTime) override;
 	virtual void StopRecording(class ULevelSequence* InSequence) override;
-	virtual TArray<UTakeRecorderSource*> PostRecording(class ULevelSequence* InSequence, class ULevelSequence* InMasterSequence, const bool bCancelled) override;
+	virtual TArray<UTakeRecorderSource*> PostRecording(class ULevelSequence* InSequence, class ULevelSequence* InRootSequence, const bool bCancelled) override;
 
 	virtual FText GetDisplayTextImpl() const override;
 	virtual void AddContentsToFolder(class UMovieSceneFolder* InFolder) override;
diff --git a/Engine/Plugins/Animation/LiveLink/Source/LiveLinkSequencer/Public/MovieSceneLiveLinkTrackRecorder.h b/Engine/Plugins/Animation/LiveLink/Source/LiveLinkSequencer/Public/MovieSceneLiveLinkTrackRecorder.h
--- a/Engine/Plugins/Animation/LiveLink/Source/LiveLinkSequencer/Public/MovieSceneLiveLinkTrackRecorder.h
+++ b/Engine/Plugins/Animation/LiveLink/Source/LiveLinkSequencer/Public/MovieSceneLiveLinkTrackRecorder.h
@@ -53,7 +53,7 @@ public:
 
 private:
 
-	UMovieSceneLiveLinkTrack* DoesLiveLinkMasterTrackExist(const FName& MasterTrackName, const TSubclassOf<ULiveLinkRole>& InTrackRole);
+	UMovieSceneLiveLinkTrack* DoesLiveLinkTrackExist(const FName& TrackName, const TSubclassOf<ULiveLinkRole>& InTrackRole);
 
 	void CreateTracks();
 
diff --git a/Engine/Plugins/Enterprise/DatasmithImporter/Source/DatasmithImporter/Private/DatasmithLevelSequenceImporter.cpp b/Engine/Plugins/Enterprise/DatasmithImporter/Source/DatasmithImporter/Private/DatasmithLevelSequenceImporter.cpp
--- a/Engine/Plugins/Enterprise/DatasmithImporter/Source/DatasmithImporter/Private/DatasmithLevelSequenceImporter.cpp
+++ b/Engine/Plugins/Enterprise/DatasmithImporter/Source/DatasmithImporter/Private/DatasmithLevelSequenceImporter.cpp
@@ -488,10 +488,10 @@ namespace DatasmithLevelSequenceImporterImpl
 
 		// Create SubTrack
 		UMovieScene* MovieScene = LevelSequence->GetMovieScene();
-		UMovieSceneSubTrack* SubTrack = MovieScene->FindMasterTrack<UMovieSceneSubTrack>();
+		UMovieSceneSubTrack* SubTrack = MovieScene->FindTrack<UMovieSceneSubTrack>();
 		if (!SubTrack)
 		{
-			SubTrack = MovieScene->AddMasterTrack<UMovieSceneSubTrack>();
+			SubTrack = MovieScene->AddTrack<UMovieSceneSubTrack>();
 		}
 		else if (SubTrack && bShouldClearReferenceTrack)
 		{
diff --git a/Engine/Plugins/Experimental/Animation/ContextualAnimation/Source/ContextualAnimationEditor/Private/ContextualAnimMovieSceneNotifyTrackEditor.cpp b/Engine/Plugins/Experimental/Animation/ContextualAnimation/Source/ContextualAnimationEditor/Private/ContextualAnimMovieSceneNotifyTrackEditor.cpp
--- a/Engine/Plugins/Experimental/Animation/ContextualAnimation/Source/ContextualAnimationEditor/Private/ContextualAnimMovieSceneNotifyTrackEditor.cpp
+++ b/Engine/Plugins/Experimental/Animation/ContextualAnimation/Source/ContextualAnimationEditor/Private/ContextualAnimMovieSceneNotifyTrackEditor.cpp
@@ -352,7 +352,7 @@ void FContextualAnimMovieSceneNotifyTrackEditor::AddNewNotifyTrack()
 	GetMovieSceneSequence().GetViewModel().AnimationModified(*Animation);
 
 	// Create and Initialize MovieSceneTrack
-	UContextualAnimMovieSceneNotifyTrack* MovieSceneTrack = GetMovieSceneSequence().GetMovieScene()->AddMasterTrack<UContextualAnimMovieSceneNotifyTrack>();
+	UContextualAnimMovieSceneNotifyTrack* MovieSceneTrack = GetMovieSceneSequence().GetMovieScene()->AddTrack<UContextualAnimMovieSceneNotifyTrack>();
 	check(MovieSceneTrack);
 
 	MovieSceneTrack->Initialize(*Animation, NewNotifyTrack);
diff --git a/Engine/Plugins/Experimental/Animation/ContextualAnimation/Source/ContextualAnimationEditor/Private/ContextualAnimViewModel.cpp b/Engine/Plugins/Experimental/Animation/ContextualAnimation/Source/ContextualAnimationEditor/Private/ContextualAnimViewModel.cpp
--- a/Engine/Plugins/Experimental/Animation/ContextualAnimation/Source/ContextualAnimationEditor/Private/ContextualAnimViewModel.cpp
+++ b/Engine/Plugins/Experimental/Animation/ContextualAnimation/Source/ContextualAnimationEditor/Private/ContextualAnimViewModel.cpp
@@ -223,10 +223,10 @@ AActor* FContextualAnimViewModel::SpawnPreviewActor(const FContextualAnimTrack&
 
 void FContextualAnimViewModel::ResetTimeline()
 {
-	for (int32 MasterTrackIdx = MovieScene->GetMasterTracks().Num() - 1; MasterTrackIdx >= 0; MasterTrackIdx--)
+	for (int32 TrackIdx = MovieScene->GetTracks().Num() - 1; TrackIdx >= 0; TrackIdx--)
 	{
-		UMovieSceneTrack& MasterTrack = *MovieScene->GetMasterTracks()[MasterTrackIdx];
-		MovieScene->RemoveMasterTrack(MasterTrack);
+		UMovieSceneTrack& Track = *MovieScene->GetTracks()[TrackIdx];
+		MovieScene->RemoveTrack(Track);
 	}
 
 	Sequencer->NotifyMovieSceneDataChanged(EMovieSceneDataChangeType::MovieSceneStructureItemsChanged);
@@ -253,7 +253,7 @@ void FContextualAnimViewModel::SetDefaultMode()
 	TArray<FName> Roles = SceneAsset->GetRoles();
 	for (FName Role : Roles)
 	{
-		UContextualAnimMovieSceneTrack* MovieSceneAnimTrack = MovieSceneSequence->GetMovieScene()->AddMasterTrack<UContextualAnimMovieSceneTrack>();
+		UContextualAnimMovieSceneTrack* MovieSceneAnimTrack = MovieSceneSequence->GetMovieScene()->AddTrack<UContextualAnimMovieSceneTrack>();
 		check(MovieSceneAnimTrack);
 
 		MovieSceneAnimTrack->Initialize(Role);
@@ -270,7 +270,7 @@ void FContextualAnimViewModel::SetDefaultMode()
 			UAnimSequenceBase* Animation = AnimTrack.Animation;
 			if (Animation)
 			{
-				UContextualAnimMovieSceneTrack* MovieSceneTrack = FindMasterTrackByRole(AnimTrack.Role);
+				UContextualAnimMovieSceneTrack* MovieSceneTrack = FindTrackByRole(AnimTrack.Role);
 				check(MovieSceneTrack)
 
 				UContextualAnimMovieSceneSection* NewSection = NewObject<UContextualAnimMovieSceneSection>(MovieSceneTrack, UContextualAnimMovieSceneSection::StaticClass(), NAME_None, RF_Transactional);
@@ -313,7 +313,7 @@ void FContextualAnimViewModel::SetNotifiesMode(const FContextualAnimTrack& AnimT
 
 	// Add Animation Track
 	{
-		UContextualAnimMovieSceneTrack* MovieSceneAnimTrack = MovieSceneSequence->GetMovieScene()->AddMasterTrack<UContextualAnimMovieSceneTrack>();
+		UContextualAnimMovieSceneTrack* MovieSceneAnimTrack = MovieSceneSequence->GetMovieScene()->AddTrack<UContextualAnimMovieSceneTrack>();
 		check(MovieSceneAnimTrack);
 
 		MovieSceneAnimTrack->SetDisplayName(FText::FromString(GetNameSafe(EditingAnimation.Get())));
@@ -335,7 +335,7 @@ void FContextualAnimViewModel::SetNotifiesMode(const FContextualAnimTrack& AnimT
 	{
 		for (const FAnimNotifyTrack& NotifyTrack : EditingAnimation->AnimNotifyTracks)
 		{
-			UContextualAnimMovieSceneNotifyTrack* Track = MovieSceneSequence->GetMovieScene()->AddMasterTrack<UContextualAnimMovieSceneNotifyTrack>();
+			UContextualAnimMovieSceneNotifyTrack* Track = MovieSceneSequence->GetMovieScene()->AddTrack<UContextualAnimMovieSceneNotifyTrack>();
 			check(Track);
 
 			Track->Initialize(*EditingAnimation, NotifyTrack);
@@ -590,12 +590,12 @@ void FContextualAnimViewModel::UpdatePreviewActorTransform(const FContextualAnim
 	}
 }
 
-UContextualAnimMovieSceneTrack* FContextualAnimViewModel::FindMasterTrackByRole(const FName& Role) const
+UContextualAnimMovieSceneTrack* FContextualAnimViewModel::FindTrackByRole(const FName& Role) const
 {
-	const TArray<UMovieSceneTrack*>& MasterTracks = MovieScene->GetMasterTracks();
-	for (UMovieSceneTrack* MasterTrack : MasterTracks)
+	const TArray<UMovieSceneTrack*>& Tracks = MovieScene->GetTracks();
+	for (UMovieSceneTrack* Track : Tracks)
 	{
-		UContextualAnimMovieSceneTrack* ContextualAnimTrack = Cast<UContextualAnimMovieSceneTrack>(MasterTrack);
+		UContextualAnimMovieSceneTrack* ContextualAnimTrack = Cast<UContextualAnimMovieSceneTrack>(Track);
 		check(ContextualAnimTrack);
 
 		if (ContextualAnimTrack->GetRole() == Role)
diff --git a/Engine/Plugins/Experimental/Animation/ContextualAnimation/Source/ContextualAnimationEditor/Public/ContextualAnimViewModel.h b/Engine/Plugins/Experimental/Animation/ContextualAnimation/Source/ContextualAnimationEditor/Public/ContextualAnimViewModel.h
--- a/Engine/Plugins/Experimental/Animation/ContextualAnimation/Source/ContextualAnimationEditor/Public/ContextualAnimViewModel.h
+++ b/Engine/Plugins/Experimental/Animation/ContextualAnimation/Source/ContextualAnimationEditor/Public/ContextualAnimViewModel.h
@@ -78,7 +78,7 @@ public:
 	UContextualAnimMovieSceneSequence* GetMovieSceneSequence() const { return MovieSceneSequence; }
 	UContextualAnimSceneAsset* GetSceneAsset() const { return SceneAsset; }
 	UContextualAnimSceneInstance* GetSceneInstance() const { return SceneInstance.Get(); }
-	UContextualAnimMovieSceneTrack* FindMasterTrackByRole(const FName& Role) const;
+	UContextualAnimMovieSceneTrack* FindTrackByRole(const FName& Role) const;
 
 	void AddNewAnimSet(const FContextualAnimNewAnimSetParams& Params);
 
diff --git a/Engine/Plugins/Experimental/ChaosCaching/Source/ChaosCachingEditor/Private/Sequencer/TakeRecorderChaosCacheSource.cpp b/Engine/Plugins/Experimental/ChaosCaching/Source/ChaosCachingEditor/Private/Sequencer/TakeRecorderChaosCacheSource.cpp
--- a/Engine/Plugins/Experimental/ChaosCaching/Source/ChaosCachingEditor/Private/Sequencer/TakeRecorderChaosCacheSource.cpp
+++ b/Engine/Plugins/Experimental/ChaosCaching/Source/ChaosCachingEditor/Private/Sequencer/TakeRecorderChaosCacheSource.cpp
@@ -1,4 +1,4 @@
-﻿// Copyright Epic Games, Inc. All Rights Reserved.
+// Copyright Epic Games, Inc. All Rights Reserved.
 
 #include "Sequencer/TakeRecorderChaosCacheSource.h"
 
@@ -24,18 +24,18 @@ UTakeRecorderChaosCacheSource::UTakeRecorderChaosCacheSource(const FObjectInitia
 	TrackTint = FColor(0.0f, 125.0f, 255.0f, 65.0f);
 }
 
-TArray<UTakeRecorderSource*> UTakeRecorderChaosCacheSource::PreRecording(ULevelSequence* InSequence, FMovieSceneSequenceID InSequenceID, ULevelSequence* InMasterSequence, FManifestSerializer* InManifestSerializer)
+TArray<UTakeRecorderSource*> UTakeRecorderChaosCacheSource::PreRecording(ULevelSequence* InSequence, FMovieSceneSequenceID InSequenceID, ULevelSequence* InRootSequence, FManifestSerializer* InManifestSerializer)
 {
 	if (ChaosCacheManager)
 	{
-		UMovieScene* MovieScene = InMasterSequence->GetMovieScene();
+		UMovieScene* MovieScene = InRootSequence->GetMovieScene();
 		TrackRecorder = NewObject<UMovieSceneChaosCacheTrackRecorder>();
 
 		// We only support possessable for now since the spawnable template is clearing out the component reference
 		// Would be better to use a TSoftObjectPtr for the actor and have a mechanism similar to the FixUpPIE to go
 		// from PIE to editor to have a correct soft object path
 		const FGuid ObjectGuid = MovieScene->AddPossessable(ChaosCacheManager->GetActorLabel(), ChaosCacheManager->GetClass());
-		InMasterSequence->BindPossessableObject(ObjectGuid, *ChaosCacheManager, ChaosCacheManager->GetWorld());
+		InRootSequence->BindPossessableObject(ObjectGuid, *ChaosCacheManager, ChaosCacheManager->GetWorld());
 		
 		TrackRecorder->CreateTrack(nullptr, ChaosCacheManager.Get(), MovieScene, nullptr, ObjectGuid);
 		CachedChaosCacheTrack = TrackRecorder->GetChaosCacheTrack();
@@ -68,7 +68,7 @@ void UTakeRecorderChaosCacheSource::StopRecording(class ULevelSequence* InSequen
 	}
 }
 
-TArray<UTakeRecorderSource*> UTakeRecorderChaosCacheSource::PostRecording(class ULevelSequence* InSequence, ULevelSequence* InMasterSequence, const bool bCancelled)
+TArray<UTakeRecorderSource*> UTakeRecorderChaosCacheSource::PostRecording(class ULevelSequence* InSequence, ULevelSequence* InRootSequence, const bool bCancelled)
 {
 	if (TrackRecorder)
 	{
@@ -83,7 +83,7 @@ void UTakeRecorderChaosCacheSource::AddContentsToFolder(UMovieSceneFolder* InFol
 {
 	if (CachedChaosCacheTrack.IsValid())
 	{
-		InFolder->AddChildMasterTrack(CachedChaosCacheTrack.Get());
+		InFolder->AddChildTrack(CachedChaosCacheTrack.Get());
 	}
 }
 
diff --git a/Engine/Plugins/Experimental/ChaosCaching/Source/ChaosCachingEditor/Public/Sequencer/TakeRecorderChaosCacheSource.h b/Engine/Plugins/Experimental/ChaosCaching/Source/ChaosCachingEditor/Public/Sequencer/TakeRecorderChaosCacheSource.h
--- a/Engine/Plugins/Experimental/ChaosCaching/Source/ChaosCachingEditor/Public/Sequencer/TakeRecorderChaosCacheSource.h
+++ b/Engine/Plugins/Experimental/ChaosCaching/Source/ChaosCachingEditor/Public/Sequencer/TakeRecorderChaosCacheSource.h
@@ -1,4 +1,4 @@
-﻿// Copyright Epic Games, Inc. All Rights Reserved.
+// Copyright Epic Games, Inc. All Rights Reserved.
 
 #pragma once
 
@@ -41,11 +41,11 @@ protected:
 
 private:
 	// ~Begin UTakeRecorderSource Interface
-	virtual TArray<UTakeRecorderSource*> PreRecording(ULevelSequence* InSequence, FMovieSceneSequenceID InSequenceID, ULevelSequence* InMasterSequence, FManifestSerializer* InManifestSerializer) override;
+	virtual TArray<UTakeRecorderSource*> PreRecording(ULevelSequence* InSequence, FMovieSceneSequenceID InSequenceID, ULevelSequence* InRootSequence, FManifestSerializer* InManifestSerializer) override;
 	virtual void StartRecording(const FTimecode& InSectionStartTimecode, const FFrameNumber& InSectionFirstFrame, class ULevelSequence* InSequence) override;
 	virtual void StopRecording(class ULevelSequence* InSequence) override;
 	virtual void TickRecording(const FQualifiedFrameTime& CurrentTime) override;
-	virtual TArray<UTakeRecorderSource*> PostRecording(class ULevelSequence* InSequence, ULevelSequence* InMasterSequence, const bool bCancelled) override;
+	virtual TArray<UTakeRecorderSource*> PostRecording(class ULevelSequence* InSequence, ULevelSequence* InRootSequence, const bool bCancelled) override;
 	virtual void AddContentsToFolder(UMovieSceneFolder* InFolder) override;
 	virtual FText GetDisplayTextImpl() const override;
 	// ~End UTakeRecorderSource
diff --git a/Engine/Plugins/Experimental/CinematicPrestreaming/Source/CinematicPrestreaming/Private/Sections/MovieSceneCinePrestreamingSection.cpp b/Engine/Plugins/Experimental/CinematicPrestreaming/Source/CinematicPrestreaming/Private/Sections/MovieSceneCinePrestreamingSection.cpp
--- a/Engine/Plugins/Experimental/CinematicPrestreaming/Source/CinematicPrestreaming/Private/Sections/MovieSceneCinePrestreamingSection.cpp
+++ b/Engine/Plugins/Experimental/CinematicPrestreaming/Source/CinematicPrestreaming/Private/Sections/MovieSceneCinePrestreamingSection.cpp
@@ -17,7 +17,7 @@ void UMovieSceneCinePrestreamingSection::ImportEntityImpl(UMovieSceneEntitySyste
 
 	OutImportedEntity->AddBuilder(
 		FEntityBuilder()
-		.AddTag(FBuiltInComponentTypes::Get()->Tags.Master)
+		.AddTag(FBuiltInComponentTypes::Get()->Tags.Root)
 		.Add(FBuiltInComponentTypes::Get()->TrackInstance, TrackInstance)
 	);
 }
diff --git a/Engine/Plugins/Experimental/CinematicPrestreaming/Source/CinematicPrestreamingEditor/Private/CinePrestreamingEditorSubsystem.cpp b/Engine/Plugins/Experimental/CinematicPrestreaming/Source/CinematicPrestreamingEditor/Private/CinePrestreamingEditorSubsystem.cpp
--- a/Engine/Plugins/Experimental/CinematicPrestreaming/Source/CinematicPrestreamingEditor/Private/CinePrestreamingEditorSubsystem.cpp
+++ b/Engine/Plugins/Experimental/CinematicPrestreaming/Source/CinematicPrestreamingEditor/Private/CinePrestreamingEditorSubsystem.cpp
@@ -2,7 +2,7 @@
 #include "CinePrestreamingEditorSubsystem.h"
 #include "CinePrestreamingData.h"
 
-#include "MoviePipelineMasterConfig.h"
+#include "MoviePipelinePrimaryConfig.h"
 #include "MoviePipelineOutputSetting.h"
 #include "MoviePipelineQueue.h"
 #include "MoviePipelineDeferredPasses.h"
diff --git a/Engine/Plugins/Experimental/CinematicPrestreaming/Source/CinematicPrestreamingEditor/Private/CinePrestreamingRecorderSetting.cpp b/Engine/Plugins/Experimental/CinematicPrestreaming/Source/CinematicPrestreamingEditor/Private/CinePrestreamingRecorderSetting.cpp
--- a/Engine/Plugins/Experimental/CinematicPrestreaming/Source/CinematicPrestreamingEditor/Private/CinePrestreamingRecorderSetting.cpp
+++ b/Engine/Plugins/Experimental/CinematicPrestreaming/Source/CinematicPrestreamingEditor/Private/CinePrestreamingRecorderSetting.cpp
@@ -9,7 +9,7 @@
 #include "Misc/FileHelper.h"
 #include "Misc/PackageName.h"
 #include "MoviePipeline.h"
-#include "MoviePipelineMasterConfig.h"
+#include "MoviePipelinePrimaryConfig.h"
 #include "MoviePipelineOutputSetting.h"
 #include "MoviePipelineQueue.h"
 #include "MoviePipelineUtils.h"
@@ -120,7 +120,7 @@ void UCinePrestreamingRecorderSetting::OnBeginFrame_GameThread()
 
 		SegmentData[ActiveShotIndex].bValid = true;
 		SegmentData[ActiveShotIndex].OutputState = GetPipeline()->GetOutputState();
-		SegmentData[ActiveShotIndex].InitialShotTick = GetPipeline()->GetActiveShotList()[ActiveShotIndex]->ShotInfo.CurrentTickInMaster;
+		SegmentData[ActiveShotIndex].InitialShotTick = GetPipeline()->GetActiveShotList()[ActiveShotIndex]->ShotInfo.CurrentTickInRoot;
 
 		PrevActiveShotIndex = ActiveShotIndex;
 	}
@@ -130,7 +130,7 @@ void UCinePrestreamingRecorderSetting::OnBeginFrame_GameThread()
 	// emulate motion blur, ie: frame 0 -> frame 1 -> frame 0 again. 
 	// So when building data it needs to be additive.
 	UMoviePipelineExecutorShot* CurrentShot = GetPipeline()->GetActiveShotList()[ActiveShotIndex];
-	const FFrameNumber FrameNumber = CurrentShot->ShotInfo.CurrentTickInMaster - SegmentData[ActiveShotIndex].InitialShotTick;
+	const FFrameNumber FrameNumber = CurrentShot->ShotInfo.CurrentTickInRoot - SegmentData[ActiveShotIndex].InitialShotTick;
 	
 	// Allocate virtual texture recording page buffer for this frame on the render thread.
 	ENQUEUE_RENDER_COMMAND(SetRecordBuffer)(
@@ -183,7 +183,7 @@ void UCinePrestreamingRecorderSetting::OnEndFrame_RenderThread()
 
 void UCinePrestreamingRecorderSetting::CreateAssetsFromData()
 {
-	UMoviePipelineOutputSetting* OutputSetting = GetPipeline()->GetPipelineMasterConfig()->FindSetting<UMoviePipelineOutputSetting>();
+	UMoviePipelineOutputSetting* OutputSetting = GetPipeline()->GetPipelinePrimaryConfig()->FindSetting<UMoviePipelineOutputSetting>();
 	FString FileNameFormatString = PackageDirectory.Path;
 
 	const bool bIncludeRenderPass = false;
@@ -254,7 +254,7 @@ void UCinePrestreamingRecorderSetting::CreateAssetsFromData()
 
 			// We only use the leafmost moviescene as that's the one that actually gets modified.
 			OutputSegment->MovieScene = Node->MovieScene.Get();
-			OutputSegment->Range = ShotInfo.TotalOutputRangeMaster * ShotInfo.OuterToInnerTransform;
+			OutputSegment->Range = ShotInfo.TotalOutputRangeRoot * ShotInfo.OuterToInnerTransform;
 		}
 
 		// Convert our samples and append them to the existing array
@@ -336,10 +336,10 @@ void UCinePrestreamingRecorderSetting::ModifyTargetSequences(const TArray<FMovie
 		Data.MovieScene->SetReadOnly(false);
 
 		// If there's an existing track we'll edit that instead.
-		UMovieSceneCinePrestreamingTrack* PrestreamingTrack = Data.MovieScene->FindMasterTrack<UMovieSceneCinePrestreamingTrack>();
+		UMovieSceneCinePrestreamingTrack* PrestreamingTrack = Data.MovieScene->FindTrack<UMovieSceneCinePrestreamingTrack>();
 		if (!PrestreamingTrack)
 		{
-			PrestreamingTrack = Data.MovieScene->AddMasterTrack<UMovieSceneCinePrestreamingTrack>();
+			PrestreamingTrack = Data.MovieScene->AddTrack<UMovieSceneCinePrestreamingTrack>();
 		}
 
 		// Remove any existing sections, as they may point to an old object, or have the wrong length, etc.
diff --git a/Engine/Plugins/Experimental/CinematicPrestreaming/Source/CinematicPrestreamingEditor/Private/TrackEditors/CinePrestreamingTrackEditor.cpp b/Engine/Plugins/Experimental/CinematicPrestreaming/Source/CinematicPrestreamingEditor/Private/TrackEditors/CinePrestreamingTrackEditor.cpp
--- a/Engine/Plugins/Experimental/CinematicPrestreaming/Source/CinematicPrestreamingEditor/Private/TrackEditors/CinePrestreamingTrackEditor.cpp
+++ b/Engine/Plugins/Experimental/CinematicPrestreaming/Source/CinematicPrestreamingEditor/Private/TrackEditors/CinePrestreamingTrackEditor.cpp
@@ -218,7 +218,7 @@ void FCinePrestreamingTrackEditor::HandleAddTrack()
 	const FScopedTransaction Transaction(LOCTEXT("AddDataLayerTrack_Transaction", "Add Cinematic Prestreaming Track"));
 	FocusedMovieScene->Modify();
 
-	UMovieSceneCinePrestreamingTrack* NewTrack = FocusedMovieScene->AddMasterTrack<UMovieSceneCinePrestreamingTrack>();
+	UMovieSceneCinePrestreamingTrack* NewTrack = FocusedMovieScene->AddTrack<UMovieSceneCinePrestreamingTrack>();
 	checkf(NewTrack, TEXT("Failed to create new cinematic prestreaming track."));
 
 	UMovieSceneCinePrestreamingSection* NewSection = AddNewSection(FocusedMovieScene, NewTrack);
diff --git a/Engine/Plugins/Experimental/CinematicPrestreaming/Source/CinematicPrestreamingEditor/Public/CinePrestreamingRecorderSetting.h b/Engine/Plugins/Experimental/CinematicPrestreaming/Source/CinematicPrestreamingEditor/Public/CinePrestreamingRecorderSetting.h
--- a/Engine/Plugins/Experimental/CinematicPrestreaming/Source/CinematicPrestreamingEditor/Public/CinePrestreamingRecorderSetting.h
+++ b/Engine/Plugins/Experimental/CinematicPrestreaming/Source/CinematicPrestreamingEditor/Public/CinePrestreamingRecorderSetting.h
@@ -47,7 +47,7 @@ public:
 protected:
 	// UMoviePipelineSetting interface
 	bool IsValidOnShots() const override { return true; }
-	bool IsValidOnMaster() const override { return true; }
+	bool IsValidOnPrimary() const override { return true; }
 #if WITH_EDITOR
 	FText GetDisplayText() const override { return NSLOCTEXT("MovieRenderPipeline", "CinePrestreamingSettingDisplayName", "Prestreaming Recorder"); }
 	FText GetFooterText(UMoviePipelineExecutorJob* InJob) const;
@@ -112,7 +112,7 @@ private:
 		TMap<FFrameNumber, FFrameData> FrameData;
 		/** OutputState required for creating output file names. */
 		FMoviePipelineFrameOutputState OutputState;
-		/** TickInMaster at first tick in frame. */
+		/** TickInRoot first tick in frame. */
 		FFrameNumber InitialShotTick;
 	};
 
diff --git a/Engine/Plugins/Experimental/LiveLinkFaceImporter/Source/LiveLinkFaceImporter/Private/LiveLinkFaceImporterFactory.cpp b/Engine/Plugins/Experimental/LiveLinkFaceImporter/Source/LiveLinkFaceImporter/Private/LiveLinkFaceImporterFactory.cpp
--- a/Engine/Plugins/Experimental/LiveLinkFaceImporter/Source/LiveLinkFaceImporter/Private/LiveLinkFaceImporterFactory.cpp
+++ b/Engine/Plugins/Experimental/LiveLinkFaceImporter/Source/LiveLinkFaceImporter/Private/LiveLinkFaceImporterFactory.cpp
@@ -180,8 +180,8 @@ UObject* ULiveLinkFaceImporterFactory::FactoryCreateText(UClass* InClass, UObjec
 	const FFrameRate TickResolution = NewLevelSequence->GetMovieScene()->GetTickResolution();
 	const FFrameRate DisplayRate = NewLevelSequence->GetMovieScene()->GetDisplayRate();
 
-	// Add the master LL track
-	UMovieSceneLiveLinkTrack* LiveLinkTrack = NewLevelSequence->GetMovieScene()->AddMasterTrack<UMovieSceneLiveLinkTrack>();
+	// Add the LL track
+	UMovieSceneLiveLinkTrack* LiveLinkTrack = NewLevelSequence->GetMovieScene()->AddTrack<UMovieSceneLiveLinkTrack>();
 	LiveLinkTrack->SetTrackRole(SubjectRole);
 	LiveLinkTrack->SetPropertyNameAndPath(SubjectName, SubjectName.ToString());
 
diff --git a/Engine/Plugins/Experimental/RenderGrid/Source/RenderGrid/Private/RenderGrid/RenderGrid.cpp b/Engine/Plugins/Experimental/RenderGrid/Source/RenderGrid/Private/RenderGrid/RenderGrid.cpp
--- a/Engine/Plugins/Experimental/RenderGrid/Source/RenderGrid/Private/RenderGrid/RenderGrid.cpp
+++ b/Engine/Plugins/Experimental/RenderGrid/Source/RenderGrid/Private/RenderGrid/RenderGrid.cpp
@@ -6,7 +6,7 @@
 #include "IRenderGridModule.h"
 #include "LevelSequence.h"
 #include "Misc/Base64.h"
-#include "MoviePipelineMasterConfig.h"
+#include "MoviePipelinePrimaryConfig.h"
 #include "MoviePipelineOutputSetting.h"
 #include "MovieScene.h"
 #include "UObject/ObjectSaveContext.h"
diff --git a/Engine/Plugins/Experimental/RenderGrid/Source/RenderGrid/Private/RenderGrid/RenderGridQueue.cpp b/Engine/Plugins/Experimental/RenderGrid/Source/RenderGrid/Private/RenderGrid/RenderGridQueue.cpp
--- a/Engine/Plugins/Experimental/RenderGrid/Source/RenderGrid/Private/RenderGrid/RenderGridQueue.cpp
+++ b/Engine/Plugins/Experimental/RenderGrid/Source/RenderGrid/Private/RenderGrid/RenderGridQueue.cpp
@@ -66,7 +66,7 @@ URenderGridMoviePipelineRenderJob* URenderGridMoviePipelineRenderJob::Create(URe
 	UMoviePipelineExecutorJob* NewJob = UMoviePipelineEditorBlueprintLibrary::CreateJobFromSequence(RenderJob->PipelineQueue, JobSequence);
 	RenderJob->PipelineExecutorJob = NewJob;
 
-	UMoviePipelineMasterConfig* JobRenderPreset = Job->GetRenderPreset();
+	UMoviePipelinePrimaryConfig* JobRenderPreset = Job->GetRenderPreset();
 	if (IsValid(JobRenderPreset))
 	{
 		NewJob->SetConfiguration(JobRenderPreset);
diff --git a/Engine/Plugins/Experimental/RenderGrid/Source/RenderGrid/Public/RenderGrid/RenderGrid.h b/Engine/Plugins/Experimental/RenderGrid/Source/RenderGrid/Public/RenderGrid/RenderGrid.h
--- a/Engine/Plugins/Experimental/RenderGrid/Source/RenderGrid/Public/RenderGrid/RenderGrid.h
+++ b/Engine/Plugins/Experimental/RenderGrid/Source/RenderGrid/Public/RenderGrid/RenderGrid.h
@@ -9,7 +9,7 @@
 
 class URenderGridQueue;
 class UMoviePipelineOutputSetting;
-class UMoviePipelineMasterConfig;
+class UMoviePipelinePrimaryConfig;
 class ULevelSequence;
 class URenderGrid;
 
@@ -198,13 +198,13 @@ public:
 
 
 	UFUNCTION(BlueprintPure, Category="Render Grid|Job")
-	UMoviePipelineMasterConfig* GetRenderPreset() const { return RenderPreset; }
+	UMoviePipelinePrimaryConfig* GetRenderPreset() const { return RenderPreset; }
 
 	UFUNCTION(BlueprintPure, Category="Render Grid|Job")
 	UMoviePipelineOutputSetting* GetRenderPresetOutputSettings() const;
 
 	UFUNCTION(BlueprintCallable, Category="Render Grid|Job")
-	void SetRenderPreset(UMoviePipelineMasterConfig* NewRenderPreset) { RenderPreset = NewRenderPreset; }
+	void SetRenderPreset(UMoviePipelinePrimaryConfig* NewRenderPreset) { RenderPreset = NewRenderPreset; }
 
 
 	bool HasRemoteControlValue(const TSharedPtr<FRemoteControlEntity>& RemoteControlEntity) const;
@@ -268,7 +268,7 @@ private:
 
 	/** The movie pipeline render preset. Render grid jobs are rendered using the movie pipeline plugin. This 'preset' contains the configuration of that plugin. */
 	UPROPERTY()
-	TObjectPtr<UMoviePipelineMasterConfig> RenderPreset;
+	TObjectPtr<UMoviePipelinePrimaryConfig> RenderPreset;
 
 	/** The remote control plugin can be used to customize and modify the way a job is rendered. If remote control is being used, the property values of this job will be stored in this map (remote control entity id -> value as bytes). */
 	UPROPERTY()
diff --git a/Engine/Plugins/Experimental/RenderGrid/Source/RenderGridEditor/Private/UI/SRenderGridJobList.cpp b/Engine/Plugins/Experimental/RenderGrid/Source/RenderGridEditor/Private/UI/SRenderGridJobList.cpp
--- a/Engine/Plugins/Experimental/RenderGrid/Source/RenderGridEditor/Private/UI/SRenderGridJobList.cpp
+++ b/Engine/Plugins/Experimental/RenderGrid/Source/RenderGridEditor/Private/UI/SRenderGridJobList.cpp
@@ -568,10 +568,10 @@ TSharedRef<SWidget> UE::RenderGrid::Private::SRenderGridJobListTableRow::Generat
 	else if (ColumnName == FRenderGridJobListColumns::RenderPreset)
 	{
 		return SNew(SObjectPropertyEntryBox)
-			.AllowedClass(UMoviePipelineMasterConfig::StaticClass())
+			.AllowedClass(UMoviePipelinePrimaryConfig::StaticClass())
 			.ObjectPath_Lambda([this]() -> FString
 			{
-				if (UMoviePipelineMasterConfig* Preset = RenderGridJob->GetRenderPreset(); IsValid(Preset))
+				if (UMoviePipelinePrimaryConfig* Preset = RenderGridJob->GetRenderPreset(); IsValid(Preset))
 				{
 					return Preset->GetPathName();
 				}
@@ -582,7 +582,7 @@ TSharedRef<SWidget> UE::RenderGrid::Private::SRenderGridJobListTableRow::Generat
 				RenderGridJob->SetRenderPreset(nullptr);
 				if (UObject* AssetDataAsset = AssetData.GetAsset(); IsValid(AssetDataAsset))
 				{
-					if (UMoviePipelineMasterConfig* Preset = Cast<UMoviePipelineMasterConfig>(AssetDataAsset))
+					if (UMoviePipelinePrimaryConfig* Preset = Cast<UMoviePipelinePrimaryConfig>(AssetDataAsset))
 					{
 						RenderGridJob->SetRenderPreset(Preset);
 					}
diff --git a/Engine/Plugins/FX/Niagara/Source/NiagaraEditor/Private/ViewModels/NiagaraSystemViewModel.cpp b/Engine/Plugins/FX/Niagara/Source/NiagaraEditor/Private/ViewModels/NiagaraSystemViewModel.cpp
--- a/Engine/Plugins/FX/Niagara/Source/NiagaraEditor/Private/ViewModels/NiagaraSystemViewModel.cpp
+++ b/Engine/Plugins/FX/Niagara/Source/NiagaraEditor/Private/ViewModels/NiagaraSystemViewModel.cpp
@@ -1749,12 +1749,12 @@ void FNiagaraSystemViewModel::ResetEmitterHandleViewModelsAndTracks()
 	if (NiagaraSequence)
 	{
 		TGuardValue<bool> UpdateGuard(bResetingSequencerTracks, true);
-		TArray<UMovieSceneTrack*> MainTracks = NiagaraSequence->GetMovieScene()->GetMasterTracks();
-		for (UMovieSceneTrack* MainTrack : MainTracks)
+		TArray<UMovieSceneTrack*> Tracks = NiagaraSequence->GetMovieScene()->GetTracks();
+		for (UMovieSceneTrack* Track : Tracks)
 		{
-			if (MainTrack != nullptr)
+			if (Track != nullptr)
 			{
-				NiagaraSequence->GetMovieScene()->RemoveMasterTrack(*MainTrack);
+				NiagaraSequence->GetMovieScene()->RemoveTrack(*Track);
 			}
 		}
 	}
@@ -1851,9 +1851,9 @@ void PopulateChildMovieSceneFoldersFromNiagaraFolders(const UNiagaraSystemEditor
 	for (const FGuid& ChildEmitterHandleId : NiagaraFolder->GetChildEmitterHandleIds())
 	{
 		UMovieSceneNiagaraEmitterTrack* const* TrackPtr = EmitterHandleIdToTrackMap.Find(ChildEmitterHandleId);
-		if (TrackPtr != nullptr && MovieSceneFolder->GetChildMasterTracks().Contains(*TrackPtr) == false)
+		if (TrackPtr != nullptr && MovieSceneFolder->GetChildTracks().Contains(*TrackPtr) == false)
 		{
-			MovieSceneFolder->AddChildMasterTrack(*TrackPtr);
+			MovieSceneFolder->AddChildTrack(*TrackPtr);
 		}
 	}
 }
@@ -1864,19 +1864,19 @@ void FNiagaraSystemViewModel::RefreshSequencerTracks()
 
 	if (Sequencer.IsValid())
 	{
-		TArray<UMovieSceneTrack*> MainTracks = NiagaraSequence->GetMovieScene()->GetMasterTracks();
-		for (UMovieSceneTrack* MainTrack : MainTracks)
+		TArray<UMovieSceneTrack*> Tracks = NiagaraSequence->GetMovieScene()->GetTracks();
+		for (UMovieSceneTrack* Track : Tracks)
 		{
-			if (MainTrack != nullptr)
+			if (Track != nullptr)
 			{
-				NiagaraSequence->GetMovieScene()->RemoveMasterTrack(*MainTrack);
+				NiagaraSequence->GetMovieScene()->RemoveTrack(*Track);
 			}
 		}
 
 		TMap<FGuid, UMovieSceneNiagaraEmitterTrack*> EmitterHandleIdToTrackMap;
 		for (TSharedRef<FNiagaraEmitterHandleViewModel> EmitterHandleViewModel : EmitterHandleViewModels)
 		{
-			UMovieSceneNiagaraEmitterTrack* EmitterTrack = Cast<UMovieSceneNiagaraEmitterTrack>(NiagaraSequence->GetMovieScene()->AddMasterTrack(UMovieSceneNiagaraEmitterTrack::StaticClass()));
+			UMovieSceneNiagaraEmitterTrack* EmitterTrack = Cast<UMovieSceneNiagaraEmitterTrack>(NiagaraSequence->GetMovieScene()->AddTrack(UMovieSceneNiagaraEmitterTrack::StaticClass()));
 			EmitterTrack->Initialize(*this, EmitterHandleViewModel, NiagaraSequence->GetMovieScene()->GetTickResolution());
 			EmitterHandleIdToTrackMap.Add(EmitterHandleViewModel->GetId(), EmitterTrack);
 		}
@@ -1905,7 +1905,7 @@ void FNiagaraSystemViewModel::UpdateSequencerTracksForEmitters(const TArray<FGui
 	if (Sequencer.IsValid())
 	{
 		TGuardValue<bool> UpdateGuard(bUpdatingSequencerFromEmitterDataChange, true);
-		for (UMovieSceneTrack* Track : NiagaraSequence->GetMovieScene()->GetMasterTracks())
+		for (UMovieSceneTrack* Track : NiagaraSequence->GetMovieScene()->GetTracks())
 		{
 			UMovieSceneNiagaraEmitterTrack* EmitterTrack = CastChecked<UMovieSceneNiagaraEmitterTrack>(Track);
 			if (EmitterIdsRequiringUpdate.Contains(EmitterTrack->GetEmitterHandleViewModel()->GetId()))
@@ -1921,7 +1921,7 @@ UMovieSceneNiagaraEmitterTrack* FNiagaraSystemViewModel::GetTrackForHandleViewMo
 {
 	if (NiagaraSequence)
 	{
-		for (UMovieSceneTrack* Track : NiagaraSequence->GetMovieScene()->GetMasterTracks())
+		for (UMovieSceneTrack* Track : NiagaraSequence->GetMovieScene()->GetTracks())
 		{
 			UMovieSceneNiagaraEmitterTrack* EmitterTrack = CastChecked<UMovieSceneNiagaraEmitterTrack>(Track);
 			if (EmitterTrack->GetEmitterHandleViewModel() == EmitterHandleViewModel)
@@ -2182,7 +2182,7 @@ void FNiagaraSystemViewModel::EmitterHandlePropertyChanged(FGuid OwningEmitterHa
 	if (bUpdatingEmittersFromSequencerDataChange == false && NiagaraSequence)
 	{
 		TGuardValue<bool> UpdateGuard(bUpdatingSequencerFromEmitterDataChange, true);
-		for (UMovieSceneTrack* Track : NiagaraSequence->GetMovieScene()->GetMasterTracks())
+		for (UMovieSceneTrack* Track : NiagaraSequence->GetMovieScene()->GetTracks())
 		{
 			UMovieSceneNiagaraEmitterTrack* EmitterTrack = CastChecked<UMovieSceneNiagaraEmitterTrack>(Track);
 			if (EmitterTrack->GetEmitterHandleViewModel()->GetId() == OwningEmitterHandleId)
@@ -2263,7 +2263,7 @@ void FNiagaraSystemViewModel::EmitterParameterStoreChanged(const FNiagaraParamet
 		else if (Sequencer.IsValid())
 		{
 			TGuardValue<bool> UpdateGuard(bUpdatingSequencerFromEmitterDataChange, true);
-			for (UMovieSceneTrack* Track : NiagaraSequence->GetMovieScene()->GetMasterTracks())
+			for (UMovieSceneTrack* Track : NiagaraSequence->GetMovieScene()->GetTracks())
 			{
 				UMovieSceneNiagaraEmitterTrack* EmitterTrack = CastChecked<UMovieSceneNiagaraEmitterTrack>(Track);
 				EmitterTrack->UpdateTrackFromEmitterParameterChange(NiagaraSequence->GetMovieScene()->GetTickResolution());
@@ -2316,7 +2316,7 @@ void PopulateNiagaraFoldersFromMovieSceneFolders(TArrayView<UMovieSceneFolder* c
 			ParentFolder->AddChildFolder(MatchingNiagaraFolder);
 		}
 
-		PopulateNiagaraFoldersFromMovieSceneFolders(MovieSceneFolder->GetChildFolders(), MovieSceneFolder->GetChildMasterTracks(), MatchingNiagaraFolder);
+		PopulateNiagaraFoldersFromMovieSceneFolders(MovieSceneFolder->GetChildFolders(), MovieSceneFolder->GetChildTracks(), MatchingNiagaraFolder);
 	}
 
 	TArray<UNiagaraSystemEditorFolder*> ChildNiagaraFolders = ParentFolder->GetChildFolders();
@@ -2370,7 +2370,7 @@ void FNiagaraSystemViewModel::SequencerDataChanged(EMovieSceneDataChangeType Dat
 		TArray<TTuple<TSharedPtr<FNiagaraEmitterHandleViewModel>, FName>> EmitterHandlesToRename;
 
 		bool bRefreshAllTracks = false;
-		for (UMovieSceneTrack* Track : NiagaraSequence->GetMovieScene()->GetMasterTracks())
+		for (UMovieSceneTrack* Track : NiagaraSequence->GetMovieScene()->GetTracks())
 		{
 			UMovieSceneNiagaraEmitterTrack* EmitterTrack = CastChecked<UMovieSceneNiagaraEmitterTrack>(Track);
 			if (EmitterTrack->GetEmitterHandleViewModel().IsValid())
@@ -2613,9 +2613,9 @@ void FNiagaraSystemViewModel::UpdateSequencerFromEmitterHandleSelection()
 	Sequencer->EmptySelection();
 	for (FGuid SelectedEmitterHandleId : SelectionViewModel->GetSelectedEmitterHandleIds())
 	{
-		for (UMovieSceneTrack* MainTrack : NiagaraSequence->GetMovieScene()->GetMasterTracks())
+		for (UMovieSceneTrack* Track : NiagaraSequence->GetMovieScene()->GetTracks())
 		{
-			UMovieSceneNiagaraEmitterTrack* EmitterTrack = Cast<UMovieSceneNiagaraEmitterTrack>(MainTrack);
+			UMovieSceneNiagaraEmitterTrack* EmitterTrack = Cast<UMovieSceneNiagaraEmitterTrack>(Track);
 			if (EmitterTrack != nullptr && EmitterTrack->GetEmitterHandleViewModel()->GetId() == SelectedEmitterHandleId)
 			{
 				Sequencer->SelectTrack(EmitterTrack);
diff --git a/Engine/Plugins/FX/NiagaraSimCaching/Source/NiagaraSimCachingEditor/Private/Sequencer/TakeRecorderNiagaraCacheSource.cpp b/Engine/Plugins/FX/NiagaraSimCaching/Source/NiagaraSimCachingEditor/Private/Sequencer/TakeRecorderNiagaraCacheSource.cpp
--- a/Engine/Plugins/FX/NiagaraSimCaching/Source/NiagaraSimCachingEditor/Private/Sequencer/TakeRecorderNiagaraCacheSource.cpp
+++ b/Engine/Plugins/FX/NiagaraSimCaching/Source/NiagaraSimCachingEditor/Private/Sequencer/TakeRecorderNiagaraCacheSource.cpp
@@ -82,7 +82,7 @@ void UTakeRecorderNiagaraCacheSource::AddContentsToFolder(UMovieSceneFolder* InF
 {
 	if (CachedNiagaraCacheTrack.IsValid())
 	{
-		InFolder->AddChildMasterTrack(CachedNiagaraCacheTrack.Get());
+		InFolder->AddChildTrack(CachedNiagaraCacheTrack.Get());
 	}
 }
 
diff --git a/Engine/Plugins/Importers/USDImporter/Source/USDExporter/Private/LevelSequenceExporterUSD.cpp b/Engine/Plugins/Importers/USDImporter/Source/USDExporter/Private/LevelSequenceExporterUSD.cpp
--- a/Engine/Plugins/Importers/USDImporter/Source/USDExporter/Private/LevelSequenceExporterUSD.cpp
+++ b/Engine/Plugins/Importers/USDImporter/Source/USDExporter/Private/LevelSequenceExporterUSD.cpp
@@ -792,7 +792,7 @@ namespace UE
 				// Also note that we can't share these bakers with our parent movie scenes in case we're a subsequence, unfortunately, because our
 				// bakers will contain lambdas that write directly to a given prim, and those prims are specific to each layer that we're exporting
 				// (e.g. our parent level sequence will export to different prims than this movie scene will)
-				for ( const UMovieSceneTrack* Track : MovieScene->GetMasterTracks() )
+				for ( const UMovieSceneTrack* Track : MovieScene->GetTracks() )
 				{
 					const UMovieSceneSubTrack* SubTrack = Cast<const UMovieSceneSubTrack>( Track );
 					if ( !SubTrack )
@@ -930,7 +930,7 @@ namespace UE
 						FString Extension;
 						FPaths::Split( UniqueFilePath, Directory, FileName, Extension );
 
-						for ( const UMovieSceneTrack* Track : MovieScene->GetMasterTracks() )
+						for ( const UMovieSceneTrack* Track : MovieScene->GetTracks() )
 						{
 							const UMovieSceneSubTrack* SubTrack = Cast<const UMovieSceneSubTrack>( Track );
 							if ( !SubTrack )
diff --git a/Engine/Plugins/Importers/USDImporter/Source/USDStage/Private/USDLevelSequenceHelper.cpp b/Engine/Plugins/Importers/USDImporter/Source/USDStage/Private/USDLevelSequenceHelper.cpp
--- a/Engine/Plugins/Importers/USDImporter/Source/USDStage/Private/USDLevelSequenceHelper.cpp
+++ b/Engine/Plugins/Importers/USDImporter/Source/USDStage/Private/USDLevelSequenceHelper.cpp
@@ -728,7 +728,7 @@ bool FUsdLevelSequenceHelperImpl::HasData() const
 		return true;
 	}
 
-	UMovieSceneSubTrack* Track = MovieScene->FindMasterTrack<UMovieSceneSubTrack>();
+	UMovieSceneSubTrack* Track = MovieScene->FindTrack<UMovieSceneSubTrack>();
 	if ( !Track )
 	{
 		return false;
@@ -1062,7 +1062,7 @@ UMovieSceneSubSection* FUsdLevelSequenceHelperImpl::FindSubSequenceSection( ULev
 		return nullptr;
 	}
 
-	UMovieSceneSubTrack* SubTrack = MovieScene->FindMasterTrack< UMovieSceneSubTrack >();
+	UMovieSceneSubTrack* SubTrack = MovieScene->FindTrack< UMovieSceneSubTrack >();
 
 	if ( !SubTrack )
 	{
@@ -1115,10 +1115,10 @@ void FUsdLevelSequenceHelperImpl::CreateSubSequenceSection( ULevelSequence& Sequ
 
 	FFrameRate TickResolution = MovieScene->GetTickResolution();
 
-	UMovieSceneSubTrack* SubTrack = MovieScene->FindMasterTrack< UMovieSceneSubTrack >();
+	UMovieSceneSubTrack* SubTrack = MovieScene->FindTrack< UMovieSceneSubTrack >();
 	if ( !SubTrack )
 	{
-		SubTrack = MovieScene->AddMasterTrack< UMovieSceneSubTrack >();
+		SubTrack = MovieScene->AddTrack< UMovieSceneSubTrack >();
 	}
 
 	const FString* LayerIdentifier = LayerIdentifierByLevelSequenceName.Find( Sequence.GetFName() );
@@ -1232,7 +1232,7 @@ void FUsdLevelSequenceHelperImpl::CreateSubSequenceSection( ULevelSequence& Sequ
 
 void FUsdLevelSequenceHelperImpl::RemoveSubSequenceSection( ULevelSequence& Sequence, ULevelSequence& SubSequence )
 {
-	if ( UMovieSceneSubTrack* SubTrack = Sequence.GetMovieScene()->FindMasterTrack< UMovieSceneSubTrack >() )
+	if ( UMovieSceneSubTrack* SubTrack = Sequence.GetMovieScene()->FindTrack< UMovieSceneSubTrack >() )
 	{
 		if ( UMovieSceneSection* SubSection = FindSubSequenceSection( Sequence, SubSequence ) )
 		{
diff --git a/Engine/Plugins/Interchange/Runtime/Source/Import/Private/Animation/InterchangeAnimationTrackSetFactory.cpp b/Engine/Plugins/Interchange/Runtime/Source/Import/Private/Animation/InterchangeAnimationTrackSetFactory.cpp
--- a/Engine/Plugins/Interchange/Runtime/Source/Import/Private/Animation/InterchangeAnimationTrackSetFactory.cpp
+++ b/Engine/Plugins/Interchange/Runtime/Source/Import/Private/Animation/InterchangeAnimationTrackSetFactory.cpp
@@ -99,7 +99,7 @@ namespace UE::Interchange::Private
 
 		FFrameNumber MinFrameNumber;
 		FFrameNumber MaxFrameNumber;
-		bool ClearSubsequenceMasterTrack = true;
+		bool ClearSubsequenceTrack = true;
 	};
 
 	void FAnimationTrackSetHelper::PopulateLevelSequence()
@@ -293,12 +293,12 @@ namespace UE::Interchange::Private
 		UMovieSceneSequence* TargetMovieSceneSequence = CastChecked<UMovieSceneSequence>(ReferenceObject.TryLoad());
 
 		// Create SubTrack
-		UMovieSceneSubTrack* SubTrack = MovieScene->FindMasterTrack<UMovieSceneSubTrack>();
+		UMovieSceneSubTrack* SubTrack = MovieScene->FindTrack<UMovieSceneSubTrack>();
 		if (!SubTrack)
 		{
-			SubTrack = MovieScene->AddMasterTrack<UMovieSceneSubTrack>();
+			SubTrack = MovieScene->AddTrack<UMovieSceneSubTrack>();
 		}
-		else if (SubTrack && ClearSubsequenceMasterTrack)
+		else if (SubTrack && ClearSubsequenceTrack)
 		{
 			SubTrack->RemoveAllAnimationData();
 		}
@@ -348,7 +348,7 @@ namespace UE::Interchange::Private
 		MinFrameNumber = FMath::Min(MinFrameNumber, NewSection->GetRange().GetLowerBoundValue());
 		MaxFrameNumber = FMath::Max(MaxFrameNumber, NewSection->GetRange().GetUpperBoundValue());
 
-		ClearSubsequenceMasterTrack = false;
+		ClearSubsequenceTrack = false;
 	}
 
 	void FAnimationTrackSetHelper::PopulateAnimationTrack(const UInterchangeAnimationTrackNode& AnimationTrackNode, int32 TrackIndex)
diff --git a/Engine/Plugins/Media/AppleProResMedia/Source/AppleProResMedia/Private/MoviePipelineAppleProResOutput.cpp b/Engine/Plugins/Media/AppleProResMedia/Source/AppleProResMedia/Private/MoviePipelineAppleProResOutput.cpp
--- a/Engine/Plugins/Media/AppleProResMedia/Source/AppleProResMedia/Private/MoviePipelineAppleProResOutput.cpp
+++ b/Engine/Plugins/Media/AppleProResMedia/Source/AppleProResMedia/Private/MoviePipelineAppleProResOutput.cpp
@@ -3,7 +3,7 @@
 #include "MoviePipelineAppleProResOutput.h"
 #include "MoviePipeline.h"
 #include "MoviePipelineOutputSetting.h"
-#include "MoviePipelineMasterConfig.h"
+#include "MoviePipelinePrimaryConfig.h"
 #include "ImagePixelData.h"
 #include "MoviePipelineImageQuantization.h"
 #include "SampleBuffer.h"
@@ -16,7 +16,7 @@
 
 TUniquePtr<MovieRenderPipeline::IVideoCodecWriter> UMoviePipelineAppleProResOutput::Initialize_GameThread(const FString& InFileName, FIntPoint InResolution, EImagePixelType InPixelType, ERGBFormat InPixelFormat, uint8 InBitDepth, uint8 InNumChannels)
 {
-	const UMoviePipelineOutputSetting* OutputSettings = GetPipeline()->GetPipelineMasterConfig()->FindSetting<UMoviePipelineOutputSetting>();
+	const UMoviePipelineOutputSetting* OutputSettings = GetPipeline()->GetPipelinePrimaryConfig()->FindSetting<UMoviePipelineOutputSetting>();
 	if (!OutputSettings)
 	{
 		return nullptr;
@@ -26,7 +26,7 @@ TUniquePtr<MovieRenderPipeline::IVideoCodecWriter> UMoviePipelineAppleProResOutp
 	Options.OutputFilename = InFileName;
 	Options.Width = InResolution.X;
 	Options.Height = InResolution.Y;
-	Options.FrameRate = GetPipeline()->GetPipelineMasterConfig()->GetEffectiveFrameRate(GetPipeline()->GetTargetSequence());
+	Options.FrameRate = GetPipeline()->GetPipelinePrimaryConfig()->GetEffectiveFrameRate(GetPipeline()->GetTargetSequence());
 	Options.MaxNumberOfEncodingThreads = bOverrideMaximumEncodingThreads ? MaxNumberOfEncodingThreads : 0; // Hardware Determine
 	Options.Codec = Codec;
 	Options.ColorPrimaries = EAppleProResEncoderColorPrimaries::CD_HDREC709; // Force Rec 709 for now
diff --git a/Engine/Plugins/Media/AvidDNxMedia/Source/Source/Private/MoviePipelineAvidDNxOutput.cpp b/Engine/Plugins/Media/AvidDNxMedia/Source/Source/Private/MoviePipelineAvidDNxOutput.cpp
--- a/Engine/Plugins/Media/AvidDNxMedia/Source/Source/Private/MoviePipelineAvidDNxOutput.cpp
+++ b/Engine/Plugins/Media/AvidDNxMedia/Source/Source/Private/MoviePipelineAvidDNxOutput.cpp
@@ -3,7 +3,7 @@
 #include "MoviePipelineAvidDNxOutput.h"
 #include "MoviePipeline.h"
 #include "MoviePipelineOutputSetting.h"
-#include "MoviePipelineMasterConfig.h"
+#include "MoviePipelinePrimaryConfig.h"
 #include "ImagePixelData.h"
 #include "MoviePipelineImageQuantization.h"
 #include "SampleBuffer.h"
@@ -16,7 +16,7 @@
 
 TUniquePtr<MovieRenderPipeline::IVideoCodecWriter> UMoviePipelineAvidDNxOutput::Initialize_GameThread(const FString& InFileName, FIntPoint InResolution, EImagePixelType InPixelType, ERGBFormat InPixelFormat, uint8 InBitDepth, uint8 InNumChannels)
 {
-	UMoviePipelineOutputSetting* OutputSettings = GetPipeline()->GetPipelineMasterConfig()->FindSetting<UMoviePipelineOutputSetting>();
+	UMoviePipelineOutputSetting* OutputSettings = GetPipeline()->GetPipelinePrimaryConfig()->FindSetting<UMoviePipelineOutputSetting>();
 	if (!OutputSettings)
 	{
 		return nullptr;
@@ -27,7 +27,7 @@ TUniquePtr<MovieRenderPipeline::IVideoCodecWriter> UMoviePipelineAvidDNxOutput::
 	Options.OutputFilename = InFileName;
 	Options.Width = InResolution.X;
 	Options.Height = InResolution.Y;
-	Options.FrameRate = GetPipeline()->GetPipelineMasterConfig()->GetEffectiveFrameRate(GetPipeline()->GetTargetSequence());
+	Options.FrameRate = GetPipeline()->GetPipelinePrimaryConfig()->GetEffectiveFrameRate(GetPipeline()->GetTargetSequence());
 	Options.bCompress = bUseCompression;
 	Options.NumberOfEncodingThreads = NumberOfEncodingThreads;
 
diff --git a/Engine/Plugins/Media/MediaCompositing/Source/MediaCompositing/Public/MovieSceneMediaTrack.h b/Engine/Plugins/Media/MediaCompositing/Source/MediaCompositing/Public/MovieSceneMediaTrack.h
--- a/Engine/Plugins/Media/MediaCompositing/Source/MediaCompositing/Public/MovieSceneMediaTrack.h
+++ b/Engine/Plugins/Media/MediaCompositing/Source/MediaCompositing/Public/MovieSceneMediaTrack.h
@@ -56,7 +56,7 @@ public:
 
 private:
 
-	/** List of all master media sections. */
+	/** List of all media sections. */
 	UPROPERTY()
 	TArray<TObjectPtr<UMovieSceneSection>> MediaSections;
 };
diff --git a/Engine/Plugins/Media/MediaCompositing/Source/MediaCompositingEditor/Private/Sequencer/MediaTrackEditor.cpp b/Engine/Plugins/Media/MediaCompositing/Source/MediaCompositingEditor/Private/Sequencer/MediaTrackEditor.cpp
--- a/Engine/Plugins/Media/MediaCompositing/Source/MediaCompositingEditor/Private/Sequencer/MediaTrackEditor.cpp
+++ b/Engine/Plugins/Media/MediaCompositing/Source/MediaCompositingEditor/Private/Sequencer/MediaTrackEditor.cpp
@@ -76,7 +76,7 @@ void FMediaTrackEditor::BuildAddTrackMenu(FMenuBuilder& MenuBuilder)
 {
 	MenuBuilder.AddMenuEntry(
 		LOCTEXT("AddTrack", "Media Track"),
-		LOCTEXT("AddTooltip", "Adds a new master media track that can play media sources."),
+		LOCTEXT("AddTooltip", "Adds a new media track that can play media sources."),
 		FSlateIcon(FAppStyle::GetAppStyleSetName(), "Sequencer.Tracks.Media"),
 		FUIAction(
 			FExecuteAction::CreateRaw(this, &FMediaTrackEditor::HandleAddMediaTrackMenuEntryExecute)
@@ -170,7 +170,7 @@ bool FMediaTrackEditor::HandleAssetAdded(UObject* Asset, const FGuid& TargetObje
 	else
 	{
 		int32 RowIndex = INDEX_NONE;
-		AnimatablePropertyChanged(FOnKeyProperty::CreateRaw(this, &FMediaTrackEditor::AddMasterMediaSource, MediaSource, RowIndex));
+		AnimatablePropertyChanged(FOnKeyProperty::CreateRaw(this, &FMediaTrackEditor::AddMediaSource, MediaSource, RowIndex));
 	}
 
 	return true;
@@ -259,11 +259,11 @@ FKeyPropertyResult FMediaTrackEditor::AddAttachedMediaSource(FFrameNumber KeyTim
 }
 
 
-FKeyPropertyResult FMediaTrackEditor::AddMasterMediaSource(FFrameNumber KeyTime, UMediaSource* MediaSource, int32 RowIndex)
+FKeyPropertyResult FMediaTrackEditor::AddMediaSource(FFrameNumber KeyTime, UMediaSource* MediaSource, int32 RowIndex)
 {
 	FKeyPropertyResult KeyPropertyResult;
 
-	FFindOrCreateMasterTrackResult<UMovieSceneMediaTrack> TrackResult = FindOrCreateMasterTrack<UMovieSceneMediaTrack>();
+	FFindOrCreateRootTrackResult<UMovieSceneMediaTrack> TrackResult = FindOrCreateRootTrack<UMovieSceneMediaTrack>();
 	UMovieSceneTrack* Track = TrackResult.Track;
 	auto MediaTrack = Cast<UMovieSceneMediaTrack>(Track);
 
@@ -337,7 +337,7 @@ void FMediaTrackEditor::HandleAddMediaTrackMenuEntryExecute()
 	const FScopedTransaction Transaction(NSLOCTEXT("Sequencer", "AddMediaTrack_Transaction", "Add Media Track"));
 	FocusedMovieScene->Modify();
 	
-	auto NewTrack = FocusedMovieScene->AddMasterTrack<UMovieSceneMediaTrack>();
+	auto NewTrack = FocusedMovieScene->AddTrack<UMovieSceneMediaTrack>();
 	ensure(NewTrack);
 	NewTrack->SetDisplayName(LOCTEXT("MediaTrackName", "Media"));
 
diff --git a/Engine/Plugins/Media/MediaCompositing/Source/MediaCompositingEditor/Private/Sequencer/MovieSceneMediaPlayerSectionRecorder.cpp b/Engine/Plugins/Media/MediaCompositing/Source/MediaCompositingEditor/Private/Sequencer/MovieSceneMediaPlayerSectionRecorder.cpp
--- a/Engine/Plugins/Media/MediaCompositing/Source/MediaCompositingEditor/Private/Sequencer/MovieSceneMediaPlayerSectionRecorder.cpp
+++ b/Engine/Plugins/Media/MediaCompositing/Source/MediaCompositingEditor/Private/Sequencer/MovieSceneMediaPlayerSectionRecorder.cpp
@@ -49,7 +49,7 @@ void FMovieSceneMediaPlayerSectionRecorder::CreateSection(UObject* InObjectToRec
 	MovieSceneTrack = InMovieScene->FindTrack<UMovieSceneMediaTrack>(InGuid, ObjectToRecord->GetFName());
 	if (!MovieSceneTrack.IsValid())
 	{
-		MovieSceneTrack = InMovieScene->AddMasterTrack<UMovieSceneMediaTrack>();
+		MovieSceneTrack = InMovieScene->AddTrack<UMovieSceneMediaTrack>();
 		MovieSceneTrack->SetDisplayName(FText::FromName(ObjectToRecord->GetFName()));
 	}
 	else
diff --git a/Engine/Plugins/Media/MediaCompositing/Source/MediaCompositingEditor/Public/Sequencer/MediaTrackEditor.h b/Engine/Plugins/Media/MediaCompositing/Source/MediaCompositingEditor/Public/Sequencer/MediaTrackEditor.h
--- a/Engine/Plugins/Media/MediaCompositing/Source/MediaCompositingEditor/Public/Sequencer/MediaTrackEditor.h
+++ b/Engine/Plugins/Media/MediaCompositing/Source/MediaCompositingEditor/Public/Sequencer/MediaTrackEditor.h
@@ -78,8 +78,8 @@ protected:
 	/** Callback for AnimatablePropertyChanged in HandleAssetAdded for attached media sources. */
 	FKeyPropertyResult AddAttachedMediaSource(FFrameNumber KeyTime, class UMediaSource* MediaSource, TArray<TWeakObjectPtr<UObject>> ObjectsToAttachTo, int32 RowIndex);
 
-	/** Callback for AnimatablePropertyChanged in HandleAssetAdded for master media sources. */
-	FKeyPropertyResult AddMasterMediaSource(FFrameNumber KeyTime, class UMediaSource* MediaSource, int32 RowIndex);
+	/** Callback for AnimatablePropertyChanged in HandleAssetAdded for media sources. */
+	FKeyPropertyResult AddMediaSource(FFrameNumber KeyTime, class UMediaSource* MediaSource, int32 RowIndex);
 
 	void AddNewSection(const FAssetData& Asset, UMovieSceneMediaTrack* Track);
 
diff --git a/Engine/Plugins/MovieScene/CustomizableSequencerTracks/Source/CustomizableSequencerTracks/Private/SequencerSectionBP.cpp b/Engine/Plugins/MovieScene/CustomizableSequencerTracks/Source/CustomizableSequencerTracks/Private/SequencerSectionBP.cpp
--- a/Engine/Plugins/MovieScene/CustomizableSequencerTracks/Source/CustomizableSequencerTracks/Private/SequencerSectionBP.cpp
+++ b/Engine/Plugins/MovieScene/CustomizableSequencerTracks/Source/CustomizableSequencerTracks/Private/SequencerSectionBP.cpp
@@ -34,7 +34,7 @@ void USequencerSectionBP::ImportEntityImpl(UMovieSceneEntitySystemLinker* Entity
 			FEntityBuilder()
 			.Add(BuiltInComponents->TrackInstance, FMovieSceneTrackInstanceComponent{ decltype(FMovieSceneTrackInstanceComponent::Owner)(this), CustomTrack->TrackInstanceType })
 			.AddConditional(BuiltInComponents->GenericObjectBinding, ObjectBindingID, ObjectBindingID.IsValid())
-			.AddTagConditional(BuiltInComponents->Tags.Master, !ObjectBindingID.IsValid())
+			.AddTagConditional(BuiltInComponents->Tags.Root, !ObjectBindingID.IsValid())
 		);
 	}
 }
\ No newline at end of file
diff --git a/Engine/Plugins/MovieScene/CustomizableSequencerTracks/Source/CustomizableSequencerTracks/Public/SequencerTrackBP.h b/Engine/Plugins/MovieScene/CustomizableSequencerTracks/Source/CustomizableSequencerTracks/Public/SequencerTrackBP.h
--- a/Engine/Plugins/MovieScene/CustomizableSequencerTracks/Source/CustomizableSequencerTracks/Public/SequencerTrackBP.h
+++ b/Engine/Plugins/MovieScene/CustomizableSequencerTracks/Source/CustomizableSequencerTracks/Public/SequencerTrackBP.h
@@ -14,7 +14,7 @@ class USequencerSectionBP;
 UENUM()
 enum class ECustomSequencerTrackType
 {
-	MasterTrack,
+	RootTrack,
 	ObjectTrack,
 };
 
diff --git a/Engine/Plugins/MovieScene/CustomizableSequencerTracks/Source/CustomizableSequencerTracksEditor/Private/SequencerTrackBPEditor.cpp b/Engine/Plugins/MovieScene/CustomizableSequencerTracks/Source/CustomizableSequencerTracksEditor/Private/SequencerTrackBPEditor.cpp
--- a/Engine/Plugins/MovieScene/CustomizableSequencerTracks/Source/CustomizableSequencerTracksEditor/Private/SequencerTrackBPEditor.cpp
+++ b/Engine/Plugins/MovieScene/CustomizableSequencerTracks/Source/CustomizableSequencerTracksEditor/Private/SequencerTrackBPEditor.cpp
@@ -140,12 +140,12 @@ void FSequencerTrackBPEditor::BuildObjectBindingTrackMenu(FMenuBuilder& MenuBuil
 
 	TArray<FAssetData> TrackTypes = DiscoverCustomTrackTypes();
 
-	// Prune tracks that do not support this object class, or are master tracks
+	// Prune tracks that do not support this object class, or are root tracks
 	FName TrackTypeTagName = GET_MEMBER_NAME_CHECKED(USequencerTrackBP, TrackType);
 	for (const FAssetData& Asset : TrackTypes)
 	{
 		FAssetDataTagMapSharedView::FFindTagResult TrackTypeTag = Asset.TagsAndValues.FindTag(TrackTypeTagName);
-		if (TrackTypeTag.IsSet() && TrackTypeTag.GetValue() == TEXT("MasterTrack"))
+		if (TrackTypeTag.IsSet() && TrackTypeTag.GetValue() == TEXT("RootTrack"))
 		{
 			continue;
 		}
@@ -183,13 +183,13 @@ void FSequencerTrackBPEditor::BuildAddTrackMenu(FMenuBuilder& MenuBuilder)
 
 	TArray<FAssetData> TrackTypes = DiscoverCustomTrackTypes();
 
-	// Only show tracks that can be master tracks
+	// Only show tracks that can be root tracks
 	FName TrackTypeTagName = GET_MEMBER_NAME_CHECKED(USequencerTrackBP, TrackType);
 	for (const FAssetData& Asset : TrackTypes)
 	{
 		FAssetDataTagMapSharedView::FFindTagResult TrackTypeTag = Asset.TagsAndValues.FindTag(TrackTypeTagName);
 
-		if (TrackTypeTag.IsSet() && TrackTypeTag.GetValue() == TEXT("MasterTrack"))
+		if (TrackTypeTag.IsSet() && TrackTypeTag.GetValue() == TEXT("RootTrack"))
 		{
 			UClass* ThisClass = LoadClassFromAssetData(Asset);
 			if (!ThisClass)
@@ -204,7 +204,7 @@ void FSequencerTrackBPEditor::BuildAddTrackMenu(FMenuBuilder& MenuBuilder)
 				FText(),
 				FSlateIcon(StyleSet.GetStyleSetName(), ThisClass->GetFName()),
 				FUIAction(
-					FExecuteAction::CreateSP(this, &FSequencerTrackBPEditor::AddNewMasterTrack, Asset)
+					FExecuteAction::CreateSP(this, &FSequencerTrackBPEditor::AddNewTrack, Asset)
 				)
 			);
 		}
@@ -333,7 +333,7 @@ bool FSequencerTrackBPEditor::SupportsType(TSubclassOf<UMovieSceneTrack> Type) c
 	return Type->IsChildOf(USequencerTrackBP::StaticClass());
 }
 
-void FSequencerTrackBPEditor::AddNewMasterTrack(FAssetData AssetData)
+void FSequencerTrackBPEditor::AddNewTrack(FAssetData AssetData)
 {
 	UMovieScene* MovieScene = GetFocusedMovieScene();
 	if (MovieScene == nullptr || MovieScene->IsReadOnly())
@@ -348,11 +348,11 @@ void FSequencerTrackBPEditor::AddNewMasterTrack(FAssetData AssetData)
 		return;
 	}
 
-	const FScopedTransaction Transaction(FText::Format(LOCTEXT("AddCustomMasterTrack_Transaction", "Add Master Track {0}"), FText::FromName(ClassToAdd->GetFName())));
+	const FScopedTransaction Transaction(FText::Format(LOCTEXT("AddCustomTrack_Transaction", "Add Track {0}"), FText::FromName(ClassToAdd->GetFName())));
 
 	MovieScene->Modify();
 
-	USequencerTrackBP* CustomTrack = CastChecked<USequencerTrackBP>(MovieScene->AddMasterTrack(ClassToAdd));
+	USequencerTrackBP* CustomTrack = CastChecked<USequencerTrackBP>(MovieScene->AddTrack(ClassToAdd));
 	CreateNewSection(CustomTrack, CustomTrack->DefaultSectionType);
 
 	if (GetSequencer().IsValid())
diff --git a/Engine/Plugins/MovieScene/CustomizableSequencerTracks/Source/CustomizableSequencerTracksEditor/Private/SequencerTrackBPEditor.h b/Engine/Plugins/MovieScene/CustomizableSequencerTracks/Source/CustomizableSequencerTracksEditor/Private/SequencerTrackBPEditor.h
--- a/Engine/Plugins/MovieScene/CustomizableSequencerTracks/Source/CustomizableSequencerTracksEditor/Private/SequencerTrackBPEditor.h
+++ b/Engine/Plugins/MovieScene/CustomizableSequencerTracks/Source/CustomizableSequencerTracksEditor/Private/SequencerTrackBPEditor.h
@@ -37,7 +37,7 @@ private:
 
 private:
 
-	void AddNewMasterTrack(FAssetData AssetData);
+	void AddNewTrack(FAssetData AssetData);
 
 	void AddNewObjectBindingTrack(FAssetData AssetData, TArray<FGuid> InObjectBindings);
 
diff --git a/Engine/Plugins/MovieScene/LevelSequenceEditor/Source/LevelSequenceEditor/Private/CinematicViewport/FilmOverlays.cpp b/Engine/Plugins/MovieScene/LevelSequenceEditor/Source/LevelSequenceEditor/Private/CinematicViewport/FilmOverlays.cpp
--- a/Engine/Plugins/MovieScene/LevelSequenceEditor/Source/LevelSequenceEditor/Private/CinematicViewport/FilmOverlays.cpp
+++ b/Engine/Plugins/MovieScene/LevelSequenceEditor/Source/LevelSequenceEditor/Private/CinematicViewport/FilmOverlays.cpp
@@ -464,13 +464,13 @@ int32 SFilmOverlay::OnPaint(const FPaintArgs& Args, const FGeometry& AllottedGeo
 
 void SFilmOverlayOptions::Construct(const FArguments& InArgs)
 {
-	MasterColorTint = FLinearColor(1.f, 1.f, 1.f, .5f);
+	PrimaryColorTint = FLinearColor(1.f, 1.f, 1.f, .5f);
 
-	MasterFilmOverlays.Add(NAME_None, TUniquePtr<IFilmOverlay>(new FFilmOverlay_None));
-	MasterFilmOverlays.Add("3x3Grid", TUniquePtr<IFilmOverlay>(new FFilmOverlay_Grid(3, 3)));
-	MasterFilmOverlays.Add("2x2Grid", TUniquePtr<IFilmOverlay>(new FFilmOverlay_Grid(2, 2)));
-	MasterFilmOverlays.Add("Crosshair", TUniquePtr<IFilmOverlay>(new FFilmOverlay_Crosshair));
-	MasterFilmOverlays.Add("Rabatment", TUniquePtr<IFilmOverlay>(new FFilmOverlay_Rabatment));
+	PrimaryFilmOverlays.Add(NAME_None, TUniquePtr<IFilmOverlay>(new FFilmOverlay_None));
+	PrimaryFilmOverlays.Add("3x3Grid", TUniquePtr<IFilmOverlay>(new FFilmOverlay_Grid(3, 3)));
+	PrimaryFilmOverlays.Add("2x2Grid", TUniquePtr<IFilmOverlay>(new FFilmOverlay_Grid(2, 2)));
+	PrimaryFilmOverlays.Add("Crosshair", TUniquePtr<IFilmOverlay>(new FFilmOverlay_Crosshair));
+	PrimaryFilmOverlays.Add("Rabatment", TUniquePtr<IFilmOverlay>(new FFilmOverlay_Rabatment));
 
 	ToggleableOverlays.Add("ActionSafeFrame",	TUniquePtr<IFilmOverlay>(new FFilmOverlay_SafeFrame(LOCTEXT("ActionSafeFrame", "Action Safe"), 95.f, FLinearColor::Red)));
 	ToggleableOverlays.Add("TitleSafeFrame",	TUniquePtr<IFilmOverlay>(new FFilmOverlay_SafeFrame(LOCTEXT("TitleSafeFrame", "Title Safe"), 90.f, FLinearColor::Yellow)));
@@ -501,17 +501,17 @@ void SFilmOverlayOptions::Construct(const FArguments& InArgs)
 	];
 }
 
-FLinearColor SFilmOverlayOptions::GetMasterColorTint() const
+FLinearColor SFilmOverlayOptions::GetPrimaryColorTint() const
 {
-	return MasterColorTint;
+	return PrimaryColorTint;
 }
 
-void SFilmOverlayOptions::OnMasterColorTintChanged(const FLinearColor& Tint)
+void SFilmOverlayOptions::OnPrimaryColorTintChanged(const FLinearColor& Tint)
 {
-	IFilmOverlay* Overlay = GetMasterFilmOverlay();
+	IFilmOverlay* Overlay = GetPrimaryFilmOverlay();
 	if (Overlay)
 	{
-		Overlay->SetTint(MasterColorTint);
+		Overlay->SetTint(PrimaryColorTint);
 	}
 }
 
@@ -537,7 +537,7 @@ TSharedRef<SWidget> SFilmOverlayOptions::GetMenuContent()
 			+ SVerticalBox::Slot()
 			.AutoHeight()
 			[
-				ConstructMasterOverlaysMenu()
+				ConstructPrimaryOverlaysMenu()
 			]
 		]
 
@@ -558,7 +558,7 @@ TSharedRef<SWidget> SFilmOverlayOptions::GetMenuContent()
 			+ SHorizontalBox::Slot()
 			.VAlign(VAlign_Center)
 			[
-				WidgetHelpers::CreateColorWidget(&MasterColorTint, FOnColorPicked::CreateRaw(this, &SFilmOverlayOptions::OnMasterColorTintChanged))
+				WidgetHelpers::CreateColorWidget(&PrimaryColorTint, FOnColorPicked::CreateRaw(this, &SFilmOverlayOptions::OnPrimaryColorTintChanged))
 			]
 		]
 
@@ -584,24 +584,24 @@ TSharedRef<SWidget> SFilmOverlayOptions::GetMenuContent()
 		];
 }
 
-TSharedRef<SWidget> SFilmOverlayOptions::ConstructMasterOverlaysMenu()
+TSharedRef<SWidget> SFilmOverlayOptions::ConstructPrimaryOverlaysMenu()
 {
 	TSharedRef<SUniformGridPanel> OverlaysPanel = SNew(SUniformGridPanel).SlotPadding(10.f);
 
 	TArray<FName> OverlayNames;
-	MasterFilmOverlays.GenerateKeyArray(OverlayNames);
+	PrimaryFilmOverlays.GenerateKeyArray(OverlayNames);
 
 	const int32 NumColumns = FMath::Log2(static_cast<float>(OverlayNames.Num() - 1));
 
 	int32 ColumnIndex = 0, RowIndex = 0;
 	for (int32 OverlayIndex = 0; OverlayIndex < OverlayNames.Num(); ++OverlayIndex)
 	{
-		IFilmOverlay& Overlay = *MasterFilmOverlays[OverlayNames[OverlayIndex]].Get();
+		IFilmOverlay& Overlay = *PrimaryFilmOverlays[OverlayNames[OverlayIndex]].Get();
 		OverlaysPanel->AddSlot(ColumnIndex, RowIndex)
 		[
 			SNew(SButton)
 			.ButtonStyle(FAppStyle::Get(), "HoverHintOnly")
-			.OnClicked(this, &SFilmOverlayOptions::SetMasterFilmOverlay, OverlayNames[OverlayIndex])
+			.OnClicked(this, &SFilmOverlayOptions::SetPrimaryFilmOverlay, OverlayNames[OverlayIndex])
 			[
 				SNew(SVerticalBox)
 
@@ -693,13 +693,13 @@ TSharedRef<SWidget> SFilmOverlayOptions::ConstructToggleableOverlaysMenu()
 	return GridPanel;
 }
 
-FReply SFilmOverlayOptions::SetMasterFilmOverlay(FName InName)
+FReply SFilmOverlayOptions::SetPrimaryFilmOverlay(FName InName)
 {
-	CurrentMasterOverlay = InName;
-	IFilmOverlay* Overlay = GetMasterFilmOverlay();
+	CurrentPrimaryOverlay = InName;
+	IFilmOverlay* Overlay = GetPrimaryFilmOverlay();
 	if (Overlay)
 	{
-		Overlay->SetTint(MasterColorTint);
+		Overlay->SetTint(PrimaryColorTint);
 	}
 	return FReply::Unhandled();
 }
@@ -713,7 +713,7 @@ TArray<IFilmOverlay*> SFilmOverlayOptions::GetActiveFilmOverlays() const
 {
 	TArray<IFilmOverlay*> Overlays;
 
-	if (IFilmOverlay* Overlay = GetMasterFilmOverlay())
+	if (IFilmOverlay* Overlay = GetPrimaryFilmOverlay())
 	{
 		Overlays.Add(Overlay);
 	}
@@ -731,19 +731,19 @@ TArray<IFilmOverlay*> SFilmOverlayOptions::GetActiveFilmOverlays() const
 
 const FSlateBrush* SFilmOverlayOptions::GetCurrentThumbnail() const
 {
-	if (!CurrentMasterOverlay.IsNone())
+	if (!CurrentPrimaryOverlay.IsNone())
 	{
-		return MasterFilmOverlays[CurrentMasterOverlay]->GetThumbnail();
+		return PrimaryFilmOverlays[CurrentPrimaryOverlay]->GetThumbnail();
 	}
 
 	return FLevelSequenceEditorStyle::Get()->GetBrush("FilmOverlay.DefaultThumbnail");
 }
 
-IFilmOverlay* SFilmOverlayOptions::GetMasterFilmOverlay() const
+IFilmOverlay* SFilmOverlayOptions::GetPrimaryFilmOverlay() const
 {
-	if (!CurrentMasterOverlay.IsNone())
+	if (!CurrentPrimaryOverlay.IsNone())
 	{
-		return MasterFilmOverlays[CurrentMasterOverlay].Get();
+		return PrimaryFilmOverlays[CurrentPrimaryOverlay].Get();
 	}
 	return nullptr;
 }
diff --git a/Engine/Plugins/MovieScene/LevelSequenceEditor/Source/LevelSequenceEditor/Private/CinematicViewport/FilmOverlays.h b/Engine/Plugins/MovieScene/LevelSequenceEditor/Source/LevelSequenceEditor/Private/CinematicViewport/FilmOverlays.h
--- a/Engine/Plugins/MovieScene/LevelSequenceEditor/Source/LevelSequenceEditor/Private/CinematicViewport/FilmOverlays.h
+++ b/Engine/Plugins/MovieScene/LevelSequenceEditor/Source/LevelSequenceEditor/Private/CinematicViewport/FilmOverlays.h
@@ -92,7 +92,7 @@ private:
 	TSharedRef<SWidget> GetMenuContent();
 
 	/** Construct the part of the menu that defines the set of film overlays */
-	TSharedRef<SWidget> ConstructMasterOverlaysMenu();
+	TSharedRef<SWidget> ConstructPrimaryOverlaysMenu();
 
 	/** Construct the part of the menu that defines the set of toggleable overlays (currently just safe-frames) */
 	TSharedRef<SWidget> ConstructToggleableOverlaysMenu();
@@ -102,29 +102,29 @@ private:
 	/** Get the thumbnail to be displayed on the combo box button */
 	const FSlateBrush* GetCurrentThumbnail() const;
 
-	/** Get the master film overlay ptr (may be nullptr) */
-	IFilmOverlay* GetMasterFilmOverlay() const;
+	/** Get the primary film overlay ptr (may be nullptr) */
+	IFilmOverlay* GetPrimaryFilmOverlay() const;
 
 	/** Get an array of all enabled film overlays */
 	TArray<IFilmOverlay*> GetActiveFilmOverlays() const;
 
-	/** Set the current master overlay to the specified name */
-	FReply SetMasterFilmOverlay(FName InName);
+	/** Set the current primary overlay to the specified name */
+	FReply SetPrimaryFilmOverlay(FName InName);
 
-	/** Get/Set the color tint override for the current master overlay */
-	FLinearColor GetMasterColorTint() const;
-	void OnMasterColorTintChanged(const FLinearColor& Tint);
+	/** Get/Set the color tint override for the current primary overlay */
+	FLinearColor GetPrimaryColorTint() const;
+	void OnPrimaryColorTintChanged(const FLinearColor& Tint);
 
 private:
 
-	/** Set of master film overlays (only one can be active at a time) */
-	TMap<FName, TUniquePtr<IFilmOverlay>> MasterFilmOverlays;
+	/** Set of primary film overlays (only one can be active at a time) */
+	TMap<FName, TUniquePtr<IFilmOverlay>> PrimaryFilmOverlays;
 
-	/** The name of the current master overlay */
-	FName CurrentMasterOverlay;
+	/** The name of the current primary overlay */
+	FName CurrentPrimaryOverlay;
 
-	/** Color tint to apply to master overlays */
-	FLinearColor MasterColorTint;
+	/** Color tint to apply to primary overlays */
+	FLinearColor PrimaryColorTint;
 
 	/** A map of toggleable overlays (any number can be active at a time) */
 	TMap<FName, TUniquePtr<IFilmOverlay>> ToggleableOverlays;
diff --git a/Engine/Plugins/MovieScene/LevelSequenceEditor/Source/LevelSequenceEditor/Private/CinematicViewport/SCinematicLevelViewport.cpp b/Engine/Plugins/MovieScene/LevelSequenceEditor/Source/LevelSequenceEditor/Private/CinematicViewport/SCinematicLevelViewport.cpp
--- a/Engine/Plugins/MovieScene/LevelSequenceEditor/Source/LevelSequenceEditor/Private/CinematicViewport/SCinematicLevelViewport.cpp
+++ b/Engine/Plugins/MovieScene/LevelSequenceEditor/Source/LevelSequenceEditor/Private/CinematicViewport/SCinematicLevelViewport.cpp
@@ -695,7 +695,7 @@ void SCinematicLevelViewport::Tick(const FGeometry& AllottedGeometry, const doub
 	}
 
 	// Find the cinematic shot track
-	UMovieSceneCinematicShotTrack* CinematicShotTrack = Cast<UMovieSceneCinematicShotTrack>(Sequence->GetMovieScene()->FindMasterTrack(UMovieSceneCinematicShotTrack::StaticClass()));
+	UMovieSceneCinematicShotTrack* CinematicShotTrack = Cast<UMovieSceneCinematicShotTrack>(Sequence->GetMovieScene()->FindTrack(UMovieSceneCinematicShotTrack::StaticClass()));
 
 	const FFrameRate OuterResolution = Sequencer->GetFocusedTickResolution();
 	const FFrameRate OuterPlayRate   = Sequencer->GetFocusedDisplayRate();
@@ -751,16 +751,16 @@ void SCinematicLevelViewport::Tick(const FGeometry& AllottedGeometry, const doub
 
 	const FMovieSceneEditorData& EditorData = Sequence->GetMovieScene()->GetEditorData();
 
-	FQualifiedFrameTime MasterStartTime(EditorData.WorkStart * OuterPlayRate, OuterPlayRate);
-	UIData.MasterStartText = FText::Format(
+	FQualifiedFrameTime RootStartTime(EditorData.WorkStart * OuterPlayRate, OuterPlayRate);
+	UIData.RootStartText = FText::Format(
 		TimeFormat,
-		FText::FromString(TimeDisplayFormatInterface->ToString(MasterStartTime.Time.GetFrame().Value))
+		FText::FromString(TimeDisplayFormatInterface->ToString(RootStartTime.Time.GetFrame().Value))
 	);
 
-	FQualifiedFrameTime MasterEndTime(EditorData.WorkEnd * OuterPlayRate, OuterPlayRate);
-	UIData.MasterEndText = FText::Format(
+	FQualifiedFrameTime RootEndTime(EditorData.WorkEnd * OuterPlayRate, OuterPlayRate);
+	UIData.RootEndText = FText::Format(
 		TimeFormat,
-		FText::FromString(TimeDisplayFormatInterface->ToString(MasterEndTime.Time.GetFrame().Value))
+		FText::FromString(TimeDisplayFormatInterface->ToString(RootEndTime.Time.GetFrame().Value))
 	);
 
 	UIData.CameraName = FText::GetEmpty();
diff --git a/Engine/Plugins/MovieScene/LevelSequenceEditor/Source/LevelSequenceEditor/Private/CinematicViewport/SCinematicLevelViewport.h b/Engine/Plugins/MovieScene/LevelSequenceEditor/Source/LevelSequenceEditor/Private/CinematicViewport/SCinematicLevelViewport.h
--- a/Engine/Plugins/MovieScene/LevelSequenceEditor/Source/LevelSequenceEditor/Private/CinematicViewport/SCinematicLevelViewport.h
+++ b/Engine/Plugins/MovieScene/LevelSequenceEditor/Source/LevelSequenceEditor/Private/CinematicViewport/SCinematicLevelViewport.h
@@ -47,14 +47,14 @@ struct FUIData
 	FText Filmback;
 	/** The text that represents the current playback time relative to the currently evaluated sequence. */
 	FText LocalPlaybackTime;
-	/** The text that represents the master start frame */
-	FText MasterStartText;
-	/** The text that represents the master end frame */
-	FText MasterEndText;
+	/** The text that represents the root start frame */
+	FText RootStartText;
+	/** The text that represents the root end frame */
+	FText RootEndText;
 
-	/** The tick resolution of the master */
+	/** The tick resolution of the root */
 	FFrameRate OuterResolution;
-	/** The play rate of the master */
+	/** The play rate of the root */
 	FFrameRate OuterPlayRate;
 };
 
diff --git a/Engine/Plugins/MovieScene/LevelSequenceEditor/Source/LevelSequenceEditor/Private/LevelSequenceEditorCommands.cpp b/Engine/Plugins/MovieScene/LevelSequenceEditor/Source/LevelSequenceEditor/Private/LevelSequenceEditorCommands.cpp
--- a/Engine/Plugins/MovieScene/LevelSequenceEditor/Source/LevelSequenceEditor/Private/LevelSequenceEditorCommands.cpp
+++ b/Engine/Plugins/MovieScene/LevelSequenceEditor/Source/LevelSequenceEditor/Private/LevelSequenceEditorCommands.cpp
@@ -12,7 +12,7 @@ FLevelSequenceEditorCommands::FLevelSequenceEditorCommands()
 void FLevelSequenceEditorCommands::RegisterCommands()
 {
 	UI_COMMAND(CreateNewLevelSequenceInLevel, "Add Level Sequence", "Create a new level sequence asset, and place an instance of it in this level", EUserInterfaceActionType::Button, FInputChord());
-	UI_COMMAND(CreateNewMasterSequenceInLevel, "Add Master Sequence", "Create a new master sequence asset, and place an instance of it in this level", EUserInterfaceActionType::Button, FInputChord());
+	UI_COMMAND(CreateNewLevelSequenceWithShotsInLevel, "Add Level Sequence with Shots", "Create a new level sequence with shots asset, and place an instance of it in this level", EUserInterfaceActionType::Button, FInputChord());
 	UI_COMMAND(ToggleCinematicViewportCommand, "Cinematic Viewport", "A viewport layout tailored to cinematic preview", EUserInterfaceActionType::RadioButton, FInputChord());
 
 	UI_COMMAND(SnapSectionsToTimelineUsingSourceTimecode, "Snap Sections to Timeline using Source Timecode", "Snap selected sections to the position in the timeline matching their source timecode", EUserInterfaceActionType::Button, FInputChord());
diff --git a/Engine/Plugins/MovieScene/LevelSequenceEditor/Source/LevelSequenceEditor/Private/LevelSequenceEditorCommands.h b/Engine/Plugins/MovieScene/LevelSequenceEditor/Source/LevelSequenceEditor/Private/LevelSequenceEditorCommands.h
--- a/Engine/Plugins/MovieScene/LevelSequenceEditor/Source/LevelSequenceEditor/Private/LevelSequenceEditorCommands.h
+++ b/Engine/Plugins/MovieScene/LevelSequenceEditor/Source/LevelSequenceEditor/Private/LevelSequenceEditorCommands.h
@@ -17,7 +17,7 @@ public:
 	virtual void RegisterCommands() override;
 
 	TSharedPtr<FUICommandInfo> CreateNewLevelSequenceInLevel;
-	TSharedPtr<FUICommandInfo> CreateNewMasterSequenceInLevel;
+	TSharedPtr<FUICommandInfo> CreateNewLevelSequenceWithShotsInLevel;
 	TSharedPtr<FUICommandInfo> ToggleCinematicViewportCommand;
 	
 	TSharedPtr<FUICommandInfo> SnapSectionsToTimelineUsingSourceTimecode;
diff --git a/Engine/Plugins/MovieScene/LevelSequenceEditor/Source/LevelSequenceEditor/Private/LevelSequenceEditorModule.cpp b/Engine/Plugins/MovieScene/LevelSequenceEditor/Source/LevelSequenceEditor/Private/LevelSequenceEditorModule.cpp
--- a/Engine/Plugins/MovieScene/LevelSequenceEditor/Source/LevelSequenceEditor/Private/LevelSequenceEditorModule.cpp
+++ b/Engine/Plugins/MovieScene/LevelSequenceEditor/Source/LevelSequenceEditor/Private/LevelSequenceEditorModule.cpp
@@ -145,8 +145,8 @@ protected:
 		CommandList->MapAction(FLevelSequenceEditorCommands::Get().CreateNewLevelSequenceInLevel,
 			FExecuteAction::CreateStatic(&FLevelSequenceEditorModule::OnCreateActorInLevel)
 		);
-		CommandList->MapAction(FLevelSequenceEditorCommands::Get().CreateNewMasterSequenceInLevel,
-			FExecuteAction::CreateStatic(&FLevelSequenceEditorModule::OnCreateMasterSequenceInLevel)
+		CommandList->MapAction(FLevelSequenceEditorCommands::Get().CreateNewLevelSequenceWithShotsInLevel,
+			FExecuteAction::CreateStatic(&FLevelSequenceEditorModule::OnCreateLevelSequenceWithShotsInLevel)
 		);
 
 		// Create and register the level editor toolbar menu extension
@@ -155,7 +155,7 @@ protected:
 			MenuBuilder.AddMenuEntry(FLevelSequenceEditorCommands::Get().CreateNewLevelSequenceInLevel);
 		}));
 		CinematicsMenuExtender->AddMenuExtension("LevelEditorNewCinematics", EExtensionHook::First, CommandList, FMenuExtensionDelegate::CreateStatic([](FMenuBuilder& MenuBuilder) {
-			MenuBuilder.AddMenuEntry(FLevelSequenceEditorCommands::Get().CreateNewMasterSequenceInLevel);
+			MenuBuilder.AddMenuEntry(FLevelSequenceEditorCommands::Get().CreateNewLevelSequenceWithShotsInLevel);
 		}));
 
 		FLevelEditorModule& LevelEditorModule = FModuleManager::LoadModuleChecked<FLevelEditorModule>("LevelEditor");
@@ -353,15 +353,15 @@ protected:
 	}
 
 	/** Callback for creating a new level sequence asset in the level. */
-	static void OnCreateMasterSequenceInLevel()
+	static void OnCreateLevelSequenceWithShotsInLevel()
 	{
 		FLevelEditorModule& LevelEditorModule = FModuleManager::GetModuleChecked<FLevelEditorModule>(TEXT("LevelEditor"));
-		LevelSequenceEditorHelpers::OpenMasterSequenceDialog(LevelEditorModule.GetLevelEditorTabManager().ToSharedRef());
+		LevelSequenceEditorHelpers::OpenLevelSequenceWithShotsDialog(LevelEditorModule.GetLevelEditorTabManager().ToSharedRef());
 	}
 
-	FOnMasterSequenceCreated& OnMasterSequenceCreated() override
+	FOnLevelSequenceWithShotsCreated& OnLevelSequenceWithShotsCreated() override
 	{
-		return OnMasterSequenceCreatedEvent;
+		return OnLevelSequenceWithShotsCreatedEvent;
 	}
 
 	FAllowPlaybackContext& OnComputePlaybackContext() override
@@ -398,7 +398,7 @@ private:
 
 	TSharedPtr<FUICommandList> CommandList;
 
-	FOnMasterSequenceCreated OnMasterSequenceCreatedEvent;
+	FOnLevelSequenceWithShotsCreated OnLevelSequenceWithShotsCreatedEvent;
 
 	FAllowPlaybackContext OnComputePlaybackContextDelegate;
 
diff --git a/Engine/Plugins/MovieScene/LevelSequenceEditor/Source/LevelSequenceEditor/Private/LevelSequenceEditorToolkit.cpp b/Engine/Plugins/MovieScene/LevelSequenceEditor/Source/LevelSequenceEditor/Private/LevelSequenceEditorToolkit.cpp
--- a/Engine/Plugins/MovieScene/LevelSequenceEditor/Source/LevelSequenceEditor/Private/LevelSequenceEditorToolkit.cpp
+++ b/Engine/Plugins/MovieScene/LevelSequenceEditor/Source/LevelSequenceEditor/Private/LevelSequenceEditorToolkit.cpp
@@ -141,7 +141,7 @@ FLevelSequenceEditorToolkit::~FLevelSequenceEditorToolkit()
 	if (FModuleManager::Get().IsModuleLoaded(TEXT("LevelSequenceEditor")))
 	{
 		auto& LevelSequenceEditorModule = FModuleManager::LoadModuleChecked<ILevelSequenceEditorModule>(TEXT("LevelSequenceEditor"));
-		LevelSequenceEditorModule.OnMasterSequenceCreated().RemoveAll(this);
+		LevelSequenceEditorModule.OnLevelSequenceWithShotsCreated().RemoveAll(this);
 	}
 
 	// unregister sequencer menu extenders
@@ -244,7 +244,7 @@ void FLevelSequenceEditorToolkit::Initialize(const EToolkitMode::Type Mode, cons
 	LevelEditorModule.OnMapChanged().AddRaw(this, &FLevelSequenceEditorToolkit::HandleMapChanged);
 
 	ILevelSequenceEditorModule& LevelSequenceEditorModule = FModuleManager::LoadModuleChecked<ILevelSequenceEditorModule>("LevelSequenceEditor");
-	LevelSequenceEditorModule.OnMasterSequenceCreated().AddRaw(this, &FLevelSequenceEditorToolkit::HandleMasterSequenceCreated);
+	LevelSequenceEditorModule.OnLevelSequenceWithShotsCreated().AddRaw(this, &FLevelSequenceEditorToolkit::HandleLevelSequenceWithShotsCreated);
 
 	FLevelSequenceEditorToolkit::OnOpened().Broadcast(*this);
 
@@ -676,25 +676,25 @@ void FLevelSequenceEditorToolkit::AddShot(UMovieSceneCinematicShotTrack* ShotTra
 	GetSequencer()->ForceEvaluate();
 	GetSequencer()->FocusSequenceInstance(*ShotSubSection);
 
-	const ULevelSequenceMasterSequenceSettings* MasterSequenceSettings = GetDefault<ULevelSequenceMasterSequenceSettings>();
+	const ULevelSequenceWithShotsSettings* LevelSequenceSettings = GetDefault<ULevelSequenceWithShotsSettings>();
 	const UMovieSceneToolsProjectSettings* ProjectSettings = GetDefault<UMovieSceneToolsProjectSettings>();
 
 	// Create any subshots
-	if (MasterSequenceSettings->SubSequenceNames.Num())
+	if (LevelSequenceSettings->SubSequenceNames.Num())
 	{
-		UMovieSceneSubTrack* SubTrack = Cast<UMovieSceneSubTrack>(ShotSequence->GetMovieScene()->FindMasterTrack(UMovieSceneSubTrack::StaticClass()));
+		UMovieSceneSubTrack* SubTrack = Cast<UMovieSceneSubTrack>(ShotSequence->GetMovieScene()->FindTrack(UMovieSceneSubTrack::StaticClass()));
 		if (!SubTrack)
 		{
-			SubTrack = Cast<UMovieSceneSubTrack>(ShotSequence->GetMovieScene()->AddMasterTrack(UMovieSceneSubTrack::StaticClass()));
+			SubTrack = Cast<UMovieSceneSubTrack>(ShotSequence->GetMovieScene()->AddTrack(UMovieSceneSubTrack::StaticClass()));
 		}
 	
 		int32 RowIndex = 0;
-		for (auto SubSequenceName : MasterSequenceSettings->SubSequenceNames)
+		for (auto SubSequenceName : LevelSequenceSettings->SubSequenceNames)
 		{
 			FString SubSequenceAssetName = ShotAssetName + ProjectSettings->SubSequenceSeparator + SubSequenceName.ToString();
 
 			UMovieSceneSequence* SubSequence = nullptr;
-			if (!MasterSequenceSettings->bInstanceSubSequences || ShotTrack->GetAllSections().Num() == 1)
+			if (!LevelSequenceSettings->bInstanceSubSequences || ShotTrack->GetAllSections().Num() == 1)
 			{
 				UObject* SubSequenceAsset = LevelSequenceEditorHelpers::CreateLevelSequenceAsset(SubSequenceAssetName, ShotPackagePath);
 				SubSequence = Cast<UMovieSceneSequence>(SubSequenceAsset);
@@ -704,7 +704,7 @@ void FLevelSequenceEditorToolkit::AddShot(UMovieSceneCinematicShotTrack* ShotTra
 				// Get the corresponding sequence from the first shot
 				UMovieSceneSubSection* FirstShotSubSection = Cast<UMovieSceneSubSection>(ShotTrack->GetAllSections()[0]);
 				UMovieSceneSequence* FirstShotSequence = FirstShotSubSection->GetSequence();
-				UMovieSceneSubTrack* FirstShotSubTrack = Cast<UMovieSceneSubTrack>(FirstShotSequence->GetMovieScene()->FindMasterTrack(UMovieSceneSubTrack::StaticClass()));
+				UMovieSceneSubTrack* FirstShotSubTrack = Cast<UMovieSceneSubTrack>(FirstShotSequence->GetMovieScene()->FindTrack(UMovieSceneSubTrack::StaticClass()));
 			
 				FString FirstShotSubSequenceAssetName = FirstShotAssetName + ProjectSettings->SubSequenceSeparator + SubSequenceName.ToString();
 
@@ -768,25 +768,25 @@ void FLevelSequenceEditorToolkit::AddShot(UMovieSceneCinematicShotTrack* ShotTra
 	}
 }
 
-void FLevelSequenceEditorToolkit::HandleMasterSequenceCreated(UObject* MasterSequenceAsset)
+void FLevelSequenceEditorToolkit::HandleLevelSequenceWithShotsCreated(UObject* LevelSequenceWithShotsAsset)
 {
-	UMovieSceneSequence* MasterSequence = Cast<UMovieSceneSequence>(MasterSequenceAsset);
-	if (!MasterSequence)
+	UMovieSceneSequence* LevelSequenceWithShots = Cast<UMovieSceneSequence>(LevelSequenceWithShotsAsset);
+	if (!LevelSequenceWithShots)
 	{
 		return;
 	}
 
-	const FScopedTransaction Transaction( LOCTEXT( "CreateMasterSequence", "Create Master Sequence" ) );
+	const FScopedTransaction Transaction( LOCTEXT( "CreateLevelSequenceWithShots", "Create Level Sequence with Shots" ) );
 	
-	const ULevelSequenceMasterSequenceSettings* MasterSequenceSettings = GetDefault<ULevelSequenceMasterSequenceSettings>();
-	uint32 NumShots = MasterSequenceSettings->MasterSequenceNumShots;
-	ULevelSequence* AssetToDuplicate = MasterSequenceSettings->MasterSequenceLevelSequenceToDuplicate.Get();
+	const ULevelSequenceWithShotsSettings* LevelSequenceSettings = GetDefault<ULevelSequenceWithShotsSettings>();
+	uint32 NumShots = LevelSequenceSettings->NumShots;
+	ULevelSequence* AssetToDuplicate = LevelSequenceSettings->SequenceToDuplicate.Get();
 
 	const UMovieSceneToolsProjectSettings* ProjectSettings = GetDefault<UMovieSceneToolsProjectSettings>();
 
-	UMovieSceneCinematicShotTrack* ShotTrack = MasterSequence->GetMovieScene()->AddMasterTrack<UMovieSceneCinematicShotTrack>();
+	UMovieSceneCinematicShotTrack* ShotTrack = LevelSequenceWithShots->GetMovieScene()->AddTrack<UMovieSceneCinematicShotTrack>();
 
-	FFrameRate TickResolution = MasterSequence->GetMovieScene()->GetTickResolution();
+	FFrameRate TickResolution = LevelSequenceWithShots->GetMovieScene()->GetTickResolution();
 
 	// Create shots with a camera cut and a camera for each
 	FFrameNumber SequenceStartTime = (ProjectSettings->DefaultStartTime * TickResolution).FloorToFrame();
@@ -799,7 +799,7 @@ void FLevelSequenceEditorToolkit::HandleMasterSequenceCreated(UObject* MasterSeq
 		ShotEndTime += ShotDuration;
 
 		FString ShotName = MovieSceneToolHelpers::GenerateNewShotName(ShotTrack->GetAllSections(), ShotStartTime);
-		FString ShotPackagePath = MovieSceneToolHelpers::GenerateNewShotPath(MasterSequence->GetMovieScene(), ShotName);
+		FString ShotPackagePath = MovieSceneToolHelpers::GenerateNewShotPath(LevelSequenceWithShots->GetMovieScene(), ShotName);
 
 		if (ShotIndex == 0)
 		{
@@ -807,24 +807,24 @@ void FLevelSequenceEditorToolkit::HandleMasterSequenceCreated(UObject* MasterSeq
 		}
 
 		AddShot(ShotTrack, ShotName, ShotPackagePath, ShotStartTime, ShotEndTime, AssetToDuplicate, FirstShotName);
-		GetSequencer()->ResetToNewRootSequence(*MasterSequence);
+		GetSequencer()->ResetToNewRootSequence(*LevelSequenceWithShots);
 
 		ShotStartTime = ShotEndTime;
 	}
 
-	MasterSequence->GetMovieScene()->SetPlaybackRange(SequenceStartTime, (ShotEndTime - SequenceStartTime).Value);
+	LevelSequenceWithShots->GetMovieScene()->SetPlaybackRange(SequenceStartTime, (ShotEndTime - SequenceStartTime).Value);
 
 #if WITH_EDITORONLY_DATA
 	const double SequenceStartSeconds = SequenceStartTime / TickResolution;
 	const double SequenceEndSeconds   = ShotEndTime / TickResolution;
 	const double OutputChange = (SequenceEndSeconds - SequenceStartSeconds) * 0.1;
 
-	FMovieSceneEditorData& EditorData = MasterSequence->GetMovieScene()->GetEditorData();
+	FMovieSceneEditorData& EditorData = LevelSequenceWithShots->GetMovieScene()->GetEditorData();
 	EditorData.ViewStart = EditorData.WorkStart = SequenceStartSeconds - OutputChange;
 	EditorData.ViewEnd   = EditorData.WorkEnd   = SequenceEndSeconds + OutputChange;
 #endif
 
-	GetSequencer()->ResetToNewRootSequence(*MasterSequence);
+	GetSequencer()->ResetToNewRootSequence(*LevelSequenceWithShots);
 
 	UActorFactory* ActorFactory = GEditor->FindActorFactoryForActorClass(ALevelSequenceActor::StaticClass());
 	if (!ensure(ActorFactory))
@@ -832,7 +832,7 @@ void FLevelSequenceEditorToolkit::HandleMasterSequenceCreated(UObject* MasterSeq
 		return;
 	}
 
-	ALevelSequenceActor* NewActor = CastChecked<ALevelSequenceActor>(GEditor->UseActorFactory(ActorFactory, FAssetData(MasterSequenceAsset), &FTransform::Identity));
+	ALevelSequenceActor* NewActor = CastChecked<ALevelSequenceActor>(GEditor->UseActorFactory(ActorFactory, FAssetData(LevelSequenceWithShotsAsset), &FTransform::Identity));
 	if (GCurrentLevelEditingViewportClient != nullptr && GCurrentLevelEditingViewportClient->IsPerspective())
 	{
 		GEditor->MoveActorInFrontOfCamera(*NewActor, GCurrentLevelEditingViewportClient->GetViewLocation(), GCurrentLevelEditingViewportClient->GetViewRotation().Vector());
diff --git a/Engine/Plugins/MovieScene/LevelSequenceEditor/Source/LevelSequenceEditor/Private/LevelSequenceEditorToolkit.h b/Engine/Plugins/MovieScene/LevelSequenceEditor/Source/LevelSequenceEditor/Private/LevelSequenceEditorToolkit.h
--- a/Engine/Plugins/MovieScene/LevelSequenceEditor/Source/LevelSequenceEditor/Private/LevelSequenceEditorToolkit.h
+++ b/Engine/Plugins/MovieScene/LevelSequenceEditor/Source/LevelSequenceEditor/Private/LevelSequenceEditorToolkit.h
@@ -112,7 +112,7 @@ protected:
 	/** Add default movie scene tracks for the given actor. */
 	void AddDefaultTracksForActor(AActor& Actor, const FGuid Binding);
 	
-	/** Add a shot to a master sequence */
+	/** Add a shot to a sequence */
 	void AddShot(UMovieSceneCinematicShotTrack* ShotTrack, const FString& ShotAssetName, const FString& ShotPackagePath, FFrameNumber ShotStartTime, FFrameNumber ShotEndTime, UObject* AssetToDuplicate, const FString& FirstShotAssetName);
 
 	/** Called whenever sequencer has received focus */
@@ -134,8 +134,8 @@ private:
 	/** Callback for map changes. */
 	void HandleMapChanged(UWorld* NewWorld, EMapChangeType MapChangeType);
 
-	/** Callback for when a master sequence is created. */
-	void HandleMasterSequenceCreated(UObject* MasterSequenceAsset);
+	/** Callback for when a sequence with shots is created. */
+	void HandleLevelSequenceWithShotsCreated(UObject* LevelSequenceWithShotsAsset);
 
 	/** Callback for the menu extensibility manager. */
 	TSharedRef<FExtender> HandleMenuExtensibilityGetExtender(const TSharedRef<FUICommandList> CommandList, const TArray<UObject*> ContextSensitiveObjects);
