PR #9635: Register additional film overlays (Contributed by [REDACTED_PERSON])

diff --git a/Engine/Plugins/MovieScene/LevelSequenceEditor/Source/LevelSequenceEditor/Private/CinematicViewport/FilmOverlays.cpp b/Engine/Plugins/MovieScene/LevelSequenceEditor/Source/LevelSequenceEditor/Private/CinematicViewport/FilmOverlays.cpp
--- a/Engine/Plugins/MovieScene/LevelSequenceEditor/Source/LevelSequenceEditor/Private/CinematicViewport/FilmOverlays.cpp
+++ b/Engine/Plugins/MovieScene/LevelSequenceEditor/Source/LevelSequenceEditor/Private/CinematicViewport/FilmOverlays.cpp
@@ -19,6 +19,7 @@
 #include "Widgets/Input/NumericTypeInterface.h"
 #include "Widgets/Input/SSpinBox.h"
 #include "Widgets/Colors/SColorPicker.h"
+#include "FilmOverlayToolkit.h"
 
 #define LOCTEXT_NAMESPACE "LevelSequenceEditorFilmOverlays"
 
@@ -466,16 +467,32 @@ void SFilmOverlayOptions::Construct(const FArguments& InArgs)
 {
 	PrimaryColorTint = FLinearColor(1.f, 1.f, 1.f, .5f);
 
-	PrimaryFilmOverlays.Add(NAME_None, TUniquePtr<IFilmOverlay>(new FFilmOverlay_None));
-	PrimaryFilmOverlays.Add("3x3Grid", TUniquePtr<IFilmOverlay>(new FFilmOverlay_Grid(3, 3)));
-	PrimaryFilmOverlays.Add("2x2Grid", TUniquePtr<IFilmOverlay>(new FFilmOverlay_Grid(2, 2)));
-	PrimaryFilmOverlays.Add("Crosshair", TUniquePtr<IFilmOverlay>(new FFilmOverlay_Crosshair));
-	PrimaryFilmOverlays.Add("Rabatment", TUniquePtr<IFilmOverlay>(new FFilmOverlay_Rabatment));
+	PrimaryOverlays.Add(MakeShareable(new FFilmOverlay_None));
+	UFilmOverlayToolkit::RegisterPrimaryFilmOverlay(NAME_None, PrimaryOverlays.Last());
 
-	ToggleableOverlays.Add("ActionSafeFrame",	TUniquePtr<IFilmOverlay>(new FFilmOverlay_SafeFrame(LOCTEXT("ActionSafeFrame", "Action Safe"), 95.f, FLinearColor::Red)));
-	ToggleableOverlays.Add("TitleSafeFrame",	TUniquePtr<IFilmOverlay>(new FFilmOverlay_SafeFrame(LOCTEXT("TitleSafeFrame", "Title Safe"), 90.f, FLinearColor::Yellow)));
-	ToggleableOverlays.Add("CustomSafeFrame",	TUniquePtr<IFilmOverlay>(new FFilmOverlay_SafeFrame(LOCTEXT("CustomSafeFrame", "Custom Safe"), 85.f, FLinearColor::Green)));
-	ToggleableOverlays.Add("LetterBox",			TUniquePtr<IFilmOverlay>(new FFilmOverlay_LetterBox));
+	PrimaryOverlays.Add(MakeShareable(new FFilmOverlay_Grid(3, 3)));
+	UFilmOverlayToolkit::RegisterPrimaryFilmOverlay("3x3Grid", PrimaryOverlays.Last());
+
+	PrimaryOverlays.Add(MakeShareable(new FFilmOverlay_Grid(2, 2)));
+	UFilmOverlayToolkit::RegisterPrimaryFilmOverlay("2x2Grid", PrimaryOverlays.Last());
+
+	PrimaryOverlays.Add(MakeShareable(new FFilmOverlay_Crosshair));
+	UFilmOverlayToolkit::RegisterPrimaryFilmOverlay("Crosshair", PrimaryOverlays.Last());
+
+	PrimaryOverlays.Add(MakeShareable(new FFilmOverlay_Rabatment));
+	UFilmOverlayToolkit::RegisterPrimaryFilmOverlay("Rabatment", PrimaryOverlays.Last());
+
+	ToggleableOverlays.Add(MakeShareable(new FFilmOverlay_SafeFrame(LOCTEXT("ActionSafeFrame", "Action Safe"), 95.f, FLinearColor::Red)));
+	UFilmOverlayToolkit::RegisterToggleableFilmOverlay("ActionSafeFrame", ToggleableOverlays.Last());
+
+	ToggleableOverlays.Add(MakeShareable(new FFilmOverlay_SafeFrame(LOCTEXT("TitleSafeFrame", "Title Safe"), 90.f, FLinearColor::Yellow)));
+	UFilmOverlayToolkit::RegisterToggleableFilmOverlay("TitleSafeFrame", ToggleableOverlays.Last());
+
+	ToggleableOverlays.Add(MakeShareable(new FFilmOverlay_SafeFrame(LOCTEXT("CustomSafeFrame", "Custom Safe"), 85.f, FLinearColor::Green)));
+	UFilmOverlayToolkit::RegisterToggleableFilmOverlay("CustomSafeFrame", ToggleableOverlays.Last());
+	
+	ToggleableOverlays.Add(MakeShareable(new FFilmOverlay_LetterBox));
+	UFilmOverlayToolkit::RegisterToggleableFilmOverlay("LetterBox",	ToggleableOverlays.Last());
 
 	OverlayWidget = SNew(SFilmOverlay)
 		.Visibility(EVisibility::HitTestInvisible)
@@ -589,14 +606,14 @@ TSharedRef<SWidget> SFilmOverlayOptions::ConstructPrimaryOverlaysMenu()
 	TSharedRef<SUniformGridPanel> OverlaysPanel = SNew(SUniformGridPanel).SlotPadding(10.f);
 
 	TArray<FName> OverlayNames;
-	PrimaryFilmOverlays.GenerateKeyArray(OverlayNames);
+	UFilmOverlayToolkit::GetPrimaryFilmOverlays().GenerateKeyArray(OverlayNames);
 
 	const int32 NumColumns = FMath::Log2(static_cast<float>(OverlayNames.Num() - 1));
 
 	int32 ColumnIndex = 0, RowIndex = 0;
 	for (int32 OverlayIndex = 0; OverlayIndex < OverlayNames.Num(); ++OverlayIndex)
 	{
-		IFilmOverlay& Overlay = *PrimaryFilmOverlays[OverlayNames[OverlayIndex]].Get();
+		IFilmOverlay& Overlay = *UFilmOverlayToolkit::GetPrimaryFilmOverlays()[OverlayNames[OverlayIndex]].Get();
 		OverlaysPanel->AddSlot(ColumnIndex, RowIndex)
 		[
 			SNew(SButton)
@@ -646,19 +663,23 @@ TSharedRef<SWidget> SFilmOverlayOptions::ConstructToggleableOverlaysMenu()
 	TSharedRef<SGridPanel> GridPanel = SNew(SGridPanel);
 
 	int32 Row = 0;
-	for (auto& Pair : ToggleableOverlays)
+	for (const TPair<FName, TSharedPtr<IFilmOverlay> >& Pair : UFilmOverlayToolkit::GetToggleableFilmOverlays())
 	{
-		TSharedPtr<SWidget> Settings = Pair.Value->ConstructSettingsWidget();
-		if (!Settings.IsValid())
+		TSharedPtr<IFilmOverlay> FilmOverlay = Pair.Value;
+		if (!FilmOverlay.IsValid())
 		{
 			continue;
 		}
 
-		TUniquePtr<IFilmOverlay>* Ptr = &Pair.Value;
+		TSharedPtr<SWidget> Settings = FilmOverlay->ConstructSettingsWidget();
+		if (!Settings.IsValid())
+		{
+			continue;		
+		}
 
-		auto OnCheckStateChanged = [=](ECheckBoxState InNewState){ (*Ptr)->SetEnabled(InNewState == ECheckBoxState::Checked); };
-		auto IsChecked = [=]{ return (*Ptr)->IsEnabled() ? ECheckBoxState::Checked : ECheckBoxState::Unchecked; };
-		auto IsEnabled = [=]{ return (*Ptr)->IsEnabled(); };
+		auto OnCheckStateChanged = [=](ECheckBoxState InNewState){ FilmOverlay->SetEnabled(InNewState == ECheckBoxState::Checked); };
+		auto IsChecked = [=]{ return FilmOverlay->IsEnabled() ? ECheckBoxState::Checked : ECheckBoxState::Unchecked; };
+		auto IsEnabled = [=]{ return FilmOverlay->IsEnabled(); };
 
 		Settings->SetEnabled(TAttribute<bool>::Create(TAttribute<bool>::FGetter::CreateLambda(IsEnabled)));
 
@@ -718,11 +739,17 @@ TArray<IFilmOverlay*> SFilmOverlayOptions::GetActiveFilmOverlays() const
 		Overlays.Add(Overlay);
 	}
 
-	for (auto& Pair: ToggleableOverlays)
+	for (const TPair<FName, TSharedPtr<IFilmOverlay> >& Pair : UFilmOverlayToolkit::GetToggleableFilmOverlays())
 	{
-		if (Pair.Value->IsEnabled())
+		TSharedPtr<IFilmOverlay> FilmOverlay = Pair.Value;
+		if (!FilmOverlay.IsValid())
+		{
+			continue;
+		}
+
+		if (FilmOverlay->IsEnabled())
 		{
-			Overlays.Add(Pair.Value.Get());
+			Overlays.Add(FilmOverlay.Get());
 		}
 	}
 
@@ -733,7 +760,7 @@ const FSlateBrush* SFilmOverlayOptions::GetCurrentThumbnail() const
 {
 	if (!CurrentPrimaryOverlay.IsNone())
 	{
-		return PrimaryFilmOverlays[CurrentPrimaryOverlay]->GetThumbnail();
+		return UFilmOverlayToolkit::GetPrimaryFilmOverlays()[CurrentPrimaryOverlay].Get()->GetThumbnail();
 	}
 
 	return FLevelSequenceEditorStyle::Get()->GetBrush("FilmOverlay.DefaultThumbnail");
@@ -743,7 +770,7 @@ IFilmOverlay* SFilmOverlayOptions::GetPrimaryFilmOverlay() const
 {
 	if (!CurrentPrimaryOverlay.IsNone())
 	{
-		return PrimaryFilmOverlays[CurrentPrimaryOverlay].Get();
+		return UFilmOverlayToolkit::GetPrimaryFilmOverlays()[CurrentPrimaryOverlay].Get();
 	}
 	return nullptr;
 }
diff --git a/Engine/Plugins/MovieScene/LevelSequenceEditor/Source/LevelSequenceEditor/Private/CinematicViewport/FilmOverlays.h b/Engine/Plugins/MovieScene/LevelSequenceEditor/Source/LevelSequenceEditor/Private/CinematicViewport/FilmOverlays.h
--- a/Engine/Plugins/MovieScene/LevelSequenceEditor/Source/LevelSequenceEditor/Private/CinematicViewport/FilmOverlays.h
+++ b/Engine/Plugins/MovieScene/LevelSequenceEditor/Source/LevelSequenceEditor/Private/CinematicViewport/FilmOverlays.h
@@ -8,48 +8,12 @@
 #include "Widgets/DeclarativeSyntaxSupport.h"
 #include "Widgets/SWidget.h"
 #include "Widgets/SCompoundWidget.h"
+#include "IFilmOverlay.h"
 
 class FPaintArgs;
 class FSlateWindowElementList;
 struct FSlateBrush;
 
-/** Abstract base class that defines how an overlay should be drawn over the viewport */
-struct IFilmOverlay
-{
-	IFilmOverlay() : Tint(FLinearColor::White), bEnabled(false) {}
-
-	virtual ~IFilmOverlay() {};
-
-	/** Get a localized display name that is representative of this overlay */
-	virtual FText GetDisplayName() const = 0;
-
-	/** Get a slate thumbnail brush that is representative of this overlay. 36x24 recommended */
-	virtual const FSlateBrush* GetThumbnail() const = 0;
-
-	/** Construct a widget that controls this overlay's arbitrary settings. Only used for toggleable overlays. */
-	virtual TSharedPtr<SWidget> ConstructSettingsWidget() { return nullptr; }
-
-	/** Paint the overlay */
-	virtual void Paint(const FGeometry& AllottedGeometry, const FSlateRect& MyCullingRect, FSlateWindowElementList& OutDrawElements, int32 LayerId) const = 0;
-
-public:
-
-	/** Get/Set a custom tint for this overlay */
-	const FLinearColor& GetTint() const { return Tint; }
-	void SetTint(const FLinearColor& InTint) { Tint = InTint; }
-
-	/** Get/Set whether this overlay should be drawn */
-	bool IsEnabled() const { return bEnabled; }
-	void SetEnabled(bool bInEnabled) { bEnabled = bInEnabled; }
-
-protected:
-
-	/** Tint to apply to this overlay */
-	FLinearColor Tint;
-
-	/** Whether this overlay is enabled or not */
-	bool bEnabled;
-};
 
 /** A widget that sits on top of a viewport, and draws custom content */
 class SFilmOverlay : public SCompoundWidget
@@ -117,17 +81,17 @@ private:
 
 private:
 
-	/** Set of primary film overlays (only one can be active at a time) */
-	TMap<FName, TUniquePtr<IFilmOverlay>> PrimaryFilmOverlays;
-
 	/** The name of the current primary overlay */
 	FName CurrentPrimaryOverlay;
 
 	/** Color tint to apply to primary overlays */
 	FLinearColor PrimaryColorTint;
 
-	/** A map of toggleable overlays (any number can be active at a time) */
-	TMap<FName, TUniquePtr<IFilmOverlay>> ToggleableOverlays;
+	/** Primary overlays (only one can be active at a time) */
+	TArray<TSharedPtr<IFilmOverlay>> PrimaryOverlays;
+
+	/** Toggleable overlays (any number can be active at a time) */
+	TArray<TSharedPtr<IFilmOverlay>> ToggleableOverlays;
 
 	/** The overlay widget we control */
 	TSharedPtr<SFilmOverlay> OverlayWidget;
diff --git a/Engine/Plugins/MovieScene/LevelSequenceEditor/Source/LevelSequenceEditor/Private/FilmOverlayToolkit.cpp b/Engine/Plugins/MovieScene/LevelSequenceEditor/Source/LevelSequenceEditor/Private/FilmOverlayToolkit.cpp
new file mode 100644
--- /dev/null
+++ b/Engine/Plugins/MovieScene/LevelSequenceEditor/Source/LevelSequenceEditor/Private/FilmOverlayToolkit.cpp
@@ -0,0 +1,36 @@
+// Copyright Epic Games, Inc. All Rights Reserved.
+
+#include "FilmOverlayToolkit.h"
+
+TMap<FName, TSharedPtr<IFilmOverlay>> UFilmOverlayToolkit::PrimaryFilmOverlays;
+TMap<FName, TSharedPtr<IFilmOverlay>> UFilmOverlayToolkit::ToggleableFilmOverlays;
+
+void UFilmOverlayToolkit::RegisterPrimaryFilmOverlay(const FName& FilmOverlayName, TSharedPtr<IFilmOverlay> FilmOverlay)
+{
+	PrimaryFilmOverlays.Add(FilmOverlayName, FilmOverlay);
+}
+
+void UFilmOverlayToolkit::UnregisterPrimaryFilmOverlay(const FName& FilmOverlayName)
+{
+	PrimaryFilmOverlays.Remove(FilmOverlayName);
+}
+
+const TMap<FName, TSharedPtr<IFilmOverlay>>& UFilmOverlayToolkit::GetPrimaryFilmOverlays()
+{
+	return PrimaryFilmOverlays;
+}
+
+void UFilmOverlayToolkit::RegisterToggleableFilmOverlay(const FName& FilmOverlayName, TSharedPtr<IFilmOverlay> FilmOverlay)
+{
+	ToggleableFilmOverlays.Add(FilmOverlayName, FilmOverlay);
+}
+
+void UFilmOverlayToolkit::UnregisterToggleableFilmOverlay(const FName& FilmOverlayName)
+{
+	ToggleableFilmOverlays.Remove(FilmOverlayName);
+}
+
+const TMap<FName, TSharedPtr<IFilmOverlay>>& UFilmOverlayToolkit::GetToggleableFilmOverlays()
+{
+	return ToggleableFilmOverlays;
+}
diff --git a/Engine/Plugins/MovieScene/LevelSequenceEditor/Source/LevelSequenceEditor/Public/FilmOverlayToolkit.h b/Engine/Plugins/MovieScene/LevelSequenceEditor/Source/LevelSequenceEditor/Public/FilmOverlayToolkit.h
new file mode 100644
--- /dev/null
+++ b/Engine/Plugins/MovieScene/LevelSequenceEditor/Source/LevelSequenceEditor/Public/FilmOverlayToolkit.h
@@ -0,0 +1,42 @@
+// Copyright Epic Games, Inc. All Rights Reserved.
+
+#pragma once
+
+#include "FilmOverlayToolkit.generated.h"
+
+struct IFilmOverlay;
+
+/** Tool kit for film overlays */
+UCLASS()
+class LEVELSEQUENCEEDITOR_API UFilmOverlayToolkit : public UObject
+{
+public:
+
+	GENERATED_BODY()
+
+	/** Register a primary film overlay */
+	static void RegisterPrimaryFilmOverlay(const FName& FilmOverlayName, TSharedPtr<IFilmOverlay> FilmOverlay);
+	
+	/** Unregister a primary film overlay */
+	static void UnregisterPrimaryFilmOverlay(const FName& FilmOverlayName);
+
+	/** Get the primary film overlays */
+	static const TMap<FName, TSharedPtr<IFilmOverlay>>& GetPrimaryFilmOverlays();
+
+	/** Register a toggleable film overlay */
+	static void RegisterToggleableFilmOverlay(const FName& FilmOverlayName, TSharedPtr<IFilmOverlay> FilmOverlay);
+
+	/** Unregister a toggleable film overlay */
+	static void UnregisterToggleableFilmOverlay(const FName& FilmOverlayName);
+
+	/** Get the toggleable film overlays */
+	static const TMap<FName, TSharedPtr<IFilmOverlay>>& GetToggleableFilmOverlays();
+
+private:
+
+	/** The primary film overlays (only one can be active at a time) */
+	static TMap<FName, TSharedPtr<IFilmOverlay>> PrimaryFilmOverlays;
+
+	/** The toggleable film overlays (any number can be active at a time) */
+	static TMap<FName, TSharedPtr<IFilmOverlay>> ToggleableFilmOverlays;
+};
diff --git a/Engine/Plugins/MovieScene/LevelSequenceEditor/Source/LevelSequenceEditor/Public/IFilmOverlay.h b/Engine/Plugins/MovieScene/LevelSequenceEditor/Source/LevelSequenceEditor/Public/IFilmOverlay.h
new file mode 100644
--- /dev/null
+++ b/Engine/Plugins/MovieScene/LevelSequenceEditor/Source/LevelSequenceEditor/Public/IFilmOverlay.h
@@ -0,0 +1,48 @@
+// Copyright Epic Games, Inc. All Rights Reserved.
+
+#pragma once
+
+#include "Layout/Geometry.h"
+#include "Math/Color.h"
+#include "Styling/SlateBrush.h"
+#include "Widgets/SWidget.h"
+
+/** Abstract base class that defines how an overlay should be drawn over the viewport */
+struct IFilmOverlay
+{
+	IFilmOverlay()
+		: Tint(FLinearColor::White) 
+		, bEnabled(false) {}
+		
+	virtual ~IFilmOverlay() {};
+
+	/** Get a localized display name that is representative of this overlay */
+	virtual FText GetDisplayName() const = 0;
+
+	/** Get a slate thumbnail brush that is representative of this overlay. 36x24 recommended */
+	virtual const FSlateBrush* GetThumbnail() const = 0;
+
+	/** Construct a widget that controls this overlay's arbitrary settings. Only used for toggleable overlays. */
+	virtual TSharedPtr<SWidget> ConstructSettingsWidget() { return nullptr; }
+
+	/** Paint the overlay */
+	virtual void Paint(const FGeometry& AllottedGeometry, const FSlateRect& MyCullingRect, FSlateWindowElementList& OutDrawElements, int32 LayerId) const = 0;
+
+public:
+
+	/** Get/Set a custom tint for this overlay */
+	const FLinearColor& GetTint() const { return Tint; }
+	void SetTint(const FLinearColor& InTint) { Tint = InTint; }
+
+	/** Get/Set whether this overlay should be drawn */
+	bool IsEnabled() const { return bEnabled; }
+	void SetEnabled(bool bInEnabled) { bEnabled = bInEnabled; }
+
+protected:
+
+	/** Tint to apply to this overlay */
+	FLinearColor Tint;
+
+	/** Whether this overlay is enabled or not */
+	bool bEnabled;
+};
