Sequencer: expose CameraShakeSourceActor, CameraShakeBase, and CameraShakePattern

diff --git a/Engine/Plugins/MovieScene/LevelSequenceEditor/Source/LevelSequenceEditor/Private/LevelSequenceEditorModule.cpp b/Engine/Plugins/MovieScene/LevelSequenceEditor/Source/LevelSequenceEditor/Private/LevelSequenceEditorModule.cpp
--- a/Engine/Plugins/MovieScene/LevelSequenceEditor/Source/LevelSequenceEditor/Private/LevelSequenceEditorModule.cpp
+++ b/Engine/Plugins/MovieScene/LevelSequenceEditor/Source/LevelSequenceEditor/Private/LevelSequenceEditorModule.cpp
@@ -17,6 +17,7 @@
 #include "Misc/LevelSequenceEditorSettings.h"
 #include "Misc/LevelSequenceEditorHelpers.h"
 #include "Styles/LevelSequenceEditorStyle.h"
+#include "Camera/CameraShakeSourceActor.h"
 #include "CineCameraActor.h"
 #include "CameraRig_Crane.h"
 #include "CameraRig_Rail.h"
@@ -174,6 +175,7 @@ protected:
 		);
 
 		IPlacementModeModule::Get().RegisterPlacementCategory(Info);
+		IPlacementModeModule::Get().RegisterPlaceableItem(Info.UniqueHandle, MakeShareable( new FPlaceableItem(nullptr, FAssetData(ACameraShakeSourceActor::StaticClass())) ));
 		IPlacementModeModule::Get().RegisterPlaceableItem(Info.UniqueHandle, MakeShareable( new FPlaceableItem(nullptr, FAssetData(ACineCameraActor::StaticClass())) ));
 		IPlacementModeModule::Get().RegisterPlaceableItem(Info.UniqueHandle, MakeShareable( new FPlaceableItem(nullptr, FAssetData(ACameraRig_Crane::StaticClass())) ));
 		IPlacementModeModule::Get().RegisterPlaceableItem(Info.UniqueHandle, MakeShareable( new FPlaceableItem(nullptr, FAssetData(ACameraRig_Rail::StaticClass())) ));
