Sequencer: Add default tracks for RelativeLocation, RelativeRotation, RelativeScale3D in property names

#rnx

diff --git a/Engine/Plugins/MovieScene/LevelSequenceEditor/Source/LevelSequenceEditor/Private/LevelSequenceEditorToolkit.cpp b/Engine/Plugins/MovieScene/LevelSequenceEditor/Source/LevelSequenceEditor/Private/LevelSequenceEditorToolkit.cpp
--- a/Engine/Plugins/MovieScene/LevelSequenceEditor/Source/LevelSequenceEditor/Private/LevelSequenceEditorToolkit.cpp
+++ b/Engine/Plugins/MovieScene/LevelSequenceEditor/Source/LevelSequenceEditor/Private/LevelSequenceEditorToolkit.cpp
@@ -364,6 +364,75 @@ void FLevelSequenceEditorToolkit::AddDefaultTracksForActor(AActor& Actor, const
 		return;
 	}
 
+
+	// Create a default section for a new track.
+	//
+	// @param InNewTrack the track to create a default section for
+	// @param InComponent for MovieScene3DTransformTrack, optional scene component to initialize the relative transform
+	auto CreateDefaultTrackSection = [this, &Actor](UMovieSceneTrack* InNewTrack, UObject* InComponent)
+	{
+		// Track class permissions can deny track creation. (UMovieScene::IsTrackClassAllowed)
+		if (!InNewTrack)
+		{
+			return;
+		}
+
+#if WITH_EDITORONLY_DATA
+		if (!InNewTrack->SupportsDefaultSections())
+		{
+			return;
+		}
+#endif
+
+		UMovieSceneSection* NewSection;
+		if (InNewTrack->GetAllSections().Num() > 0)
+		{
+			NewSection = InNewTrack->GetAllSections()[0];
+		}
+		else
+		{
+			NewSection = InNewTrack->CreateNewSection();
+			InNewTrack->AddSection(*NewSection);
+		}
+
+		// @todo sequencer: hack: setting defaults for transform tracks
+		if (InNewTrack->IsA(UMovieScene3DTransformTrack::StaticClass()) && Sequencer->GetAutoSetTrackDefaults())
+		{
+			auto TransformSection = Cast<UMovieScene3DTransformSection>(NewSection);
+
+			FVector Location = Actor.GetActorLocation();
+			FRotator Rotation = Actor.GetActorRotation();
+			FVector Scale = Actor.GetActorScale();
+
+			if (USceneComponent* SceneComponent = Cast<USceneComponent>(InComponent))
+			{
+				FTransform ActorRelativeTransform = SceneComponent->GetRelativeTransform();
+
+				Location = ActorRelativeTransform.GetTranslation();
+				Rotation = ActorRelativeTransform.GetRotation().Rotator();
+				Scale = ActorRelativeTransform.GetScale3D();
+			}
+
+			TArrayView<FMovieSceneDoubleChannel*> DoubleChannels = TransformSection->GetChannelProxy().GetChannels<FMovieSceneDoubleChannel>();
+			DoubleChannels[0]->SetDefault(Location.X);
+			DoubleChannels[1]->SetDefault(Location.Y);
+			DoubleChannels[2]->SetDefault(Location.Z);
+
+			DoubleChannels[3]->SetDefault(Rotation.Euler().X);
+			DoubleChannels[4]->SetDefault(Rotation.Euler().Y);
+			DoubleChannels[5]->SetDefault(Rotation.Euler().Z);
+
+			DoubleChannels[6]->SetDefault(Scale.X);
+			DoubleChannels[7]->SetDefault(Scale.Y);
+			DoubleChannels[8]->SetDefault(Scale.Z);
+		}
+
+		if (GetSequencer()->GetInfiniteKeyAreas())
+		{
+			NewSection->SetRange(TRange<FFrameNumber>::All());
+		}
+	};	
+	
 	// add default tracks
 	for (const FLevelSequenceTrackSettings& TrackSettings : GetDefault<ULevelSequenceEditorSettings>()->TrackSettings)
 	{
@@ -405,67 +474,7 @@ void FLevelSequenceEditorToolkit::AddDefaultTracksForActor(AActor& Actor, const
 				if (!NewTrack)
 				{
 					NewTrack = MovieScene->AddTrack(TrackClass, Binding);
-				}
-
-				if (!NewTrack)
-				{
-					continue;
-				}
-
-				bool bCreateDefaultSection = false;
-#if WITH_EDITORONLY_DATA
-				bCreateDefaultSection = NewTrack->SupportsDefaultSections();
-#endif
-
-				if (bCreateDefaultSection)
-				{
-					UMovieSceneSection* NewSection;
-					if (NewTrack->GetAllSections().Num() > 0)
-					{
-						NewSection = NewTrack->GetAllSections()[0];
-					}
-					else
-					{
-						NewSection = NewTrack->CreateNewSection();
-						NewTrack->AddSection(*NewSection);
-					}
-
-					// @todo sequencer: hack: setting defaults for transform tracks
-					if (NewTrack->IsA(UMovieScene3DTransformTrack::StaticClass()) && Sequencer->GetAutoSetTrackDefaults())
-					{
-						auto TransformSection = Cast<UMovieScene3DTransformSection>(NewSection);
-
-						FVector Location = Actor.GetActorLocation();
-						FRotator Rotation = Actor.GetActorRotation();
-						FVector Scale = Actor.GetActorScale();
-
-						if (Actor.GetRootComponent())
-						{
-							FTransform ActorRelativeTransform = Actor.GetRootComponent()->GetRelativeTransform();
-
-							Location = ActorRelativeTransform.GetTranslation();
-							Rotation = ActorRelativeTransform.GetRotation().Rotator();
-							Scale = ActorRelativeTransform.GetScale3D();
-						}
-
-						TArrayView<FMovieSceneDoubleChannel*> DoubleChannels = TransformSection->GetChannelProxy().GetChannels<FMovieSceneDoubleChannel>();
-						DoubleChannels[0]->SetDefault(Location.X);
-						DoubleChannels[1]->SetDefault(Location.Y);
-						DoubleChannels[2]->SetDefault(Location.Z);
-
-						DoubleChannels[3]->SetDefault(Rotation.Euler().X);
-						DoubleChannels[4]->SetDefault(Rotation.Euler().Y);
-						DoubleChannels[5]->SetDefault(Rotation.Euler().Z);
-
-						DoubleChannels[6]->SetDefault(Scale.X);
-						DoubleChannels[7]->SetDefault(Scale.Y);
-						DoubleChannels[8]->SetDefault(Scale.Z);
-					}
-
-					if (GetSequencer()->GetInfiniteKeyAreas())
-					{
-						NewSection->SetRange(TRange<FFrameNumber>::All());
-					}
+					CreateDefaultTrackSection(NewTrack, Actor.GetRootComponent());
 				}
 			}
 		}
@@ -535,6 +544,7 @@ void FLevelSequenceEditorToolkit::AddDefaultTracksForActor(AActor& Actor, const
 			TArray<FString> PropertyNames;
 			PropertyTrackSettings.PropertyPath.ParseIntoArray(PropertyNames, TEXT("."));
 
+			bool bReplaceWithTransformTrack = false;
 			for (const FString& PropertyName : PropertyNames)
 			{
 				// skip past excluded properties
@@ -549,6 +559,14 @@ void FLevelSequenceEditorToolkit::AddDefaultTracksForActor(AActor& Actor, const
 				if (Property != nullptr)
 				{
 					PropertyPath->AddProperty(FPropertyInfo(Property));
+
+					// Transform tracks are a special case and must be handled separately.
+					if (PropertyOwner->IsA(USceneComponent::StaticClass()) &&
+						(PropertyName == TEXT("RelativeLocation") || PropertyName == TEXT("RelativeRotation") || PropertyName == TEXT("RelativeScale3D")))
+					{
+						bReplaceWithTransformTrack = true;
+						break;
+					}
 				}
 
 				FStructProperty* StructProperty = CastField<FStructProperty>(Property);
@@ -570,6 +588,19 @@ void FLevelSequenceEditorToolkit::AddDefaultTracksForActor(AActor& Actor, const
 				break;
 			}
 
+			if (bReplaceWithTransformTrack)
+			{
+				FGuid ComponentBinding = Sequencer->GetHandleToObject(PropertyOwner);
+				UClass* TrackClass = UMovieScene3DTransformTrack::StaticClass();
+				UMovieSceneTrack* NewTrack = MovieScene->FindTrack(TrackClass, ComponentBinding);
+				if (!NewTrack)
+				{
+					NewTrack = MovieScene->AddTrack(TrackClass, ComponentBinding);
+					CreateDefaultTrackSection(NewTrack, PropertyOwner);
+				}
+				continue;
+			}
+
 			if (!Sequencer->CanKeyProperty(FCanKeyPropertyParams(PropertyOwner->GetClass(), *PropertyPath)))
 			{
 				continue;
