Sequencer: Clamp to playback range when scrubbing in cinematic viewport if enabled

diff --git a/Engine/Plugins/MovieScene/LevelSequenceEditor/Source/LevelSequenceEditor/Private/CinematicViewport/SCinematicTransportRange.cpp b/Engine/Plugins/MovieScene/LevelSequenceEditor/Source/LevelSequenceEditor/Private/CinematicViewport/SCinematicTransportRange.cpp
--- a/Engine/Plugins/MovieScene/LevelSequenceEditor/Source/LevelSequenceEditor/Private/CinematicViewport/SCinematicTransportRange.cpp
+++ b/Engine/Plugins/MovieScene/LevelSequenceEditor/Source/LevelSequenceEditor/Private/CinematicViewport/SCinematicTransportRange.cpp
@@ -3,7 +3,9 @@
 #include "CinematicViewport/SCinematicTransportRange.h"
 #include "Rendering/DrawElements.h"
 #include "SequencerKeyCollection.h"
+#include "SequencerSettings.h"
 #include "MovieSceneSequence.h"
+#include "MovieSceneTimeHelpers.h"
 #include "Styling/AppStyle.h"
 #include "Styles/LevelSequenceEditorStyle.h"
 #include "MovieScene.h"
@@ -83,7 +85,16 @@ void SCinematicTransportRange::SetTime(const FGeometry& MyGeometry, const FPoint
 		FMovieSceneEditorData& EditorData = Sequencer->GetFocusedMovieSceneSequence()->GetMovieScene()->GetEditorData();
 		double NewTimeSeconds = EditorData.ViewStart + (EditorData.ViewEnd - EditorData.ViewStart) * Lerp;
 
-		Sequencer->SetLocalTime(NewTimeSeconds * Sequencer->GetFocusedTickResolution(), ESnapTimeMode::STM_All);
+		FFrameTime ScrubTime = NewTimeSeconds * Sequencer->GetFocusedTickResolution();
+
+		// Clamp first, snap to frame last
+		if (Sequencer->GetSequencerSettings()->ShouldKeepCursorInPlayRangeWhileScrubbing())
+		{
+			TRange<FFrameNumber> PlaybackRange = Sequencer->GetFocusedMovieSceneSequence()->GetMovieScene()->GetPlaybackRange();
+			ScrubTime = UE::MovieScene::ClampToDiscreteRange(ScrubTime, PlaybackRange);
+		}
+
+		Sequencer->SetLocalTime(ScrubTime, ESnapTimeMode::STM_All);
 	}
 }
 
