Sequencer: Override BringToolkitToFront so that it can draw attention to the Sequencer tab as opposed to just bringing the whole level editor to front. This allows a hidden Sequencer tab to be brought to the front when Sequencer is invoked.

diff --git a/Engine/Plugins/MovieScene/LevelSequenceEditor/Source/LevelSequenceEditor/Private/LevelSequenceEditorToolkit.cpp b/Engine/Plugins/MovieScene/LevelSequenceEditor/Source/LevelSequenceEditor/Private/LevelSequenceEditorToolkit.cpp
--- a/Engine/Plugins/MovieScene/LevelSequenceEditor/Source/LevelSequenceEditor/Private/LevelSequenceEditorToolkit.cpp
+++ b/Engine/Plugins/MovieScene/LevelSequenceEditor/Source/LevelSequenceEditor/Private/LevelSequenceEditorToolkit.cpp
@@ -314,6 +314,16 @@ FText FLevelSequenceEditorToolkit::GetTabSuffix() const
 	return FText::GetEmpty();
 }
 
+void FLevelSequenceEditorToolkit::BringToolkitToFront()
+{
+	ILevelSequenceEditorToolkit::BringToolkitToFront();
+
+	FLevelEditorModule& LevelEditorModule = FModuleManager::LoadModuleChecked<FLevelEditorModule>("LevelEditor");
+	TSharedPtr<FTabManager> LevelEditorTabManager = LevelEditorModule.GetLevelEditorTabManager();
+	LevelEditorTabManager->TryInvokeTab(LevelEditorTabIds::Sequencer);
+}
+
+
 /* FLevelSequenceEditorToolkit implementation
  *****************************************************************************/
 
diff --git a/Engine/Plugins/MovieScene/LevelSequenceEditor/Source/LevelSequenceEditor/Private/LevelSequenceEditorToolkit.h b/Engine/Plugins/MovieScene/LevelSequenceEditor/Source/LevelSequenceEditor/Private/LevelSequenceEditorToolkit.h
--- a/Engine/Plugins/MovieScene/LevelSequenceEditor/Source/LevelSequenceEditor/Private/LevelSequenceEditorToolkit.h
+++ b/Engine/Plugins/MovieScene/LevelSequenceEditor/Source/LevelSequenceEditor/Private/LevelSequenceEditorToolkit.h
@@ -106,6 +106,7 @@ public:
 	virtual FLinearColor GetWorldCentricTabColorScale() const override;
 	virtual FString GetWorldCentricTabPrefix() const override;
 	virtual FText GetTabSuffix() const override;
+	virtual void BringToolkitToFront() override;
 
 protected:
 
diff --git a/Engine/Source/Editor/LevelEditor/Public/LevelEditor.h b/Engine/Source/Editor/LevelEditor/Public/LevelEditor.h
--- a/Engine/Source/Editor/LevelEditor/Public/LevelEditor.h
+++ b/Engine/Source/Editor/LevelEditor/Public/LevelEditor.h
@@ -28,7 +28,7 @@ extern const FName LevelEditorApp;
 
 DECLARE_DELEGATE_RetVal_OneParam(bool, FAreObjectsEditable, const TArray<TWeakObjectPtr<UObject>>&);
 
-class LevelEditorTabIds
+class LEVELEDITOR_API LevelEditorTabIds
 {
 public:
 	static const FName LevelEditorViewport;
