Cinematic Viewport: Command bindings for composition overlays

Added IFilmOverlay::GetToolTip()

diff --git a/Engine/Plugins/MovieScene/LevelSequenceEditor/Source/LevelSequenceEditor/Private/CinematicViewport/CinematicViewportCommands.cpp b/Engine/Plugins/MovieScene/LevelSequenceEditor/Source/LevelSequenceEditor/Private/CinematicViewport/CinematicViewportCommands.cpp
new file mode 100644
--- /dev/null
+++ b/Engine/Plugins/MovieScene/LevelSequenceEditor/Source/LevelSequenceEditor/Private/CinematicViewport/CinematicViewportCommands.cpp
@@ -0,0 +1,25 @@
+// Copyright Epic Games, Inc. All Rights Reserved.
+
+#include "CinematicViewportCommands.h"
+
+#include "Framework/Commands/InputChord.h"
+#include "Framework/Commands/UICommandInfo.h"
+#include "GenericPlatform/GenericApplication.h"
+#include "InputCoreTypes.h"
+
+#define LOCTEXT_NAMESPACE "CinematicViewportCommands"
+
+void FCinematicViewportCommands::RegisterCommands()
+{
+	UI_COMMAND( Disabled, "Disabled", "Disable the composition overlay", EUserInterfaceActionType::Button, FInputChord(EModifierKey::Control | EModifierKey::Shift, EKeys::NumPadZero) );
+	UI_COMMAND( Grid3x3, "Grid 3x3", "Enable the grid (3x3) composition overlay", EUserInterfaceActionType::Button, FInputChord(EModifierKey::Shift | EModifierKey::Control, EKeys::NumPadSix) );
+	UI_COMMAND( Grid2x2, "Grid 2x2", "Enable the grid (2x2) composition overlay", EUserInterfaceActionType::Button, FInputChord(EModifierKey::Shift | EModifierKey::Control, EKeys::NumPadSeven) );
+	UI_COMMAND( Crosshair, "Crosshair", "Enable the crosshair composition overlay", EUserInterfaceActionType::Button, FInputChord(EModifierKey::Shift | EModifierKey::Control, EKeys::NumPadEight) );
+	UI_COMMAND( Rabatment, "Rabatment", "Enable the rabatment composition overlay", EUserInterfaceActionType::Button, FInputChord(EModifierKey::Shift | EModifierKey::Control, EKeys::NumPadNine));
+	UI_COMMAND( ActionSafe, "Action Safe", "Enable the action safe composition overlay", EUserInterfaceActionType::Button, FInputChord(EModifierKey::Shift | EModifierKey::Control, EKeys::NumPadOne));
+	UI_COMMAND( TitleSafe, "Title Safe", "Enable the title safe composition overlay", EUserInterfaceActionType::Button, FInputChord(EModifierKey::Shift | EModifierKey::Control, EKeys::NumPadTwo));
+	UI_COMMAND( CustomSafe, "Custom Safe", "Enable the custom safe composition overlay", EUserInterfaceActionType::Button, FInputChord(EModifierKey::Shift | EModifierKey::Control, EKeys::NumPadThree));
+	UI_COMMAND( Letterbox, "Letterbox", "Enable the letterbox composition overlay", EUserInterfaceActionType::Button, FInputChord(EModifierKey::Shift | EModifierKey::Control, EKeys::NumPadFour));
+}
+
+#undef LOCTEXT_NAMESPACE
diff --git a/Engine/Plugins/MovieScene/LevelSequenceEditor/Source/LevelSequenceEditor/Private/CinematicViewport/CinematicViewportCommands.h b/Engine/Plugins/MovieScene/LevelSequenceEditor/Source/LevelSequenceEditor/Private/CinematicViewport/CinematicViewportCommands.h
new file mode 100644
--- /dev/null
+++ b/Engine/Plugins/MovieScene/LevelSequenceEditor/Source/LevelSequenceEditor/Private/CinematicViewport/CinematicViewportCommands.h
@@ -0,0 +1,59 @@
+// Copyright Epic Games, Inc. All Rights Reserved.
+
+#pragma once
+
+#include "CoreMinimal.h"
+#include "Framework/Commands/Commands.h"
+#include "Internationalization/Internationalization.h"
+#include "Styling/AppStyle.h"
+#include "Templates/SharedPointer.h"
+#include "UObject/NameTypes.h"
+#include "UObject/UnrealNames.h"
+
+class FUICommandInfo;
+
+class FCinematicViewportCommands : public TCommands<FCinematicViewportCommands>
+{
+
+public:
+	FCinematicViewportCommands() : TCommands<FCinematicViewportCommands>
+	(
+		"Cinematic Viewport",
+		NSLOCTEXT("Contexts", "CinematicViewport", "Cinematic Viewport"),
+		NAME_None,
+		FAppStyle::GetAppStyleSetName() // Icon Style Set
+	)
+	{}
+	
+	/** Disabled */
+	TSharedPtr<FUICommandInfo> Disabled;
+
+	/** Grid 2x2 */
+	TSharedPtr<FUICommandInfo> Grid2x2;
+
+	/** Grid 3x3 */
+	TSharedPtr<FUICommandInfo> Grid3x3;
+
+	/** Crosshair */
+	TSharedPtr<FUICommandInfo> Crosshair;
+
+	/** Rabatment */
+	TSharedPtr<FUICommandInfo> Rabatment;
+
+	/** ActionSafe */
+	TSharedPtr<FUICommandInfo> ActionSafe;
+
+	/** TitleSafe */
+	TSharedPtr<FUICommandInfo> TitleSafe;
+
+	/** CustomSafe */
+	TSharedPtr<FUICommandInfo> CustomSafe;
+
+	/** Letterbox */
+	TSharedPtr<FUICommandInfo> Letterbox;
+
+	/**
+	 * Initialize commands
+	 */
+	virtual void RegisterCommands() override;
+};
diff --git a/Engine/Plugins/MovieScene/LevelSequenceEditor/Source/LevelSequenceEditor/Private/CinematicViewport/FilmOverlays.cpp b/Engine/Plugins/MovieScene/LevelSequenceEditor/Source/LevelSequenceEditor/Private/CinematicViewport/FilmOverlays.cpp
--- a/Engine/Plugins/MovieScene/LevelSequenceEditor/Source/LevelSequenceEditor/Private/CinematicViewport/FilmOverlays.cpp
+++ b/Engine/Plugins/MovieScene/LevelSequenceEditor/Source/LevelSequenceEditor/Private/CinematicViewport/FilmOverlays.cpp
@@ -1,6 +1,7 @@
 // Copyright Epic Games, Inc. All Rights Reserved.
 
 #include "CinematicViewport/FilmOverlays.h"
+#include "CinematicViewportCommands.h"
 #include "Widgets/Layout/SBox.h"
 #include "Widgets/Layout/SHeader.h"
 #include "Widgets/Layout/SGridPanel.h"
@@ -76,6 +77,11 @@ struct FFilmOverlay_None : IFilmOverlay
 {
 	FText GetDisplayName() const { return LOCTEXT("OverlayDisabled", "Disabled"); }
 	
+	FText GetToolTip() const
+	{
+		return FCinematicViewportCommands::Get().Disabled.Get()->GetInputText();
+	}
+
 	const FSlateBrush* GetThumbnail() const { return FLevelSequenceEditorStyle::Get()->GetBrush("FilmOverlay.Disabled"); }
 
 	void Paint(const FGeometry& AllottedGeometry, const FSlateRect& MyCullingRect, FSlateWindowElementList& OutDrawElements, int32 LayerId) const
@@ -101,6 +107,22 @@ struct FFilmOverlay_Grid : IFilmOverlay
 		);
 	}
 	
+	FText GetToolTip() const
+	{
+		if (NumDivsH == 2 && NumDivsV == 2)
+		{
+			return FCinematicViewportCommands::Get().Grid2x2.Get()->GetInputText();
+		}
+		else if (NumDivsH == 3 && NumDivsV == 3)
+		{
+			return FCinematicViewportCommands::Get().Grid3x3.Get()->GetInputText();
+		}
+		else
+		{
+			return GetDisplayName();
+		}
+	}
+
 	const FSlateBrush* GetThumbnail() const
 	{
 		return FLevelSequenceEditorStyle::Get()->GetBrush(BrushName);
@@ -156,7 +178,9 @@ private:
 struct FFilmOverlay_Rabatment : IFilmOverlay
 {
 	FText GetDisplayName() const { return LOCTEXT("RabatmentName", "Rabatment"); }
-	
+
+	FText GetToolTip() const { return FCinematicViewportCommands::Get().Rabatment.Get()->GetInputText(); }
+
 	const FSlateBrush* GetThumbnail() const { return FLevelSequenceEditorStyle::Get()->GetBrush("FilmOverlay.Rabatment"); }
 
 	void Paint(const FGeometry& AllottedGeometry, const FSlateRect& MyCullingRect, FSlateWindowElementList& OutDrawElements, int32 LayerId) const
@@ -197,6 +221,7 @@ struct FFilmOverlay_Rabatment : IFilmOverlay
 struct FFilmOverlay_Crosshair : IFilmOverlay
 {
 	FText GetDisplayName() const { return LOCTEXT("CrosshairName", "Crosshair"); }
+	FText GetToolTip() const { return FCinematicViewportCommands::Get().Crosshair.Get()->GetInputText(); }
 	const FSlateBrush* GetThumbnail() const { return FLevelSequenceEditorStyle::Get()->GetBrush("FilmOverlay.Crosshair"); }
 	void Paint(const FGeometry& AllottedGeometry, const FSlateRect& MyCullingRect, FSlateWindowElementList& OutDrawElements, int32 LayerId) const
 	{
@@ -263,6 +288,22 @@ struct FFilmOverlay_SafeFrame : IFilmOverlay
 	}
 
 	FText GetDisplayName() const { return DisplayName; }
+	FText GetToolTip() const
+	{
+		if (GetDisplayName().EqualTo(FCinematicViewportCommands::Get().ActionSafe->GetLabel()))
+		{
+			return FCinematicViewportCommands::Get().ActionSafe.Get()->GetInputText();
+		}
+		else if (GetDisplayName().EqualTo(FCinematicViewportCommands::Get().TitleSafe->GetLabel()))
+		{
+			return FCinematicViewportCommands::Get().TitleSafe.Get()->GetInputText();
+		}
+		else if (GetDisplayName().EqualTo(FCinematicViewportCommands::Get().CustomSafe->GetLabel()))
+		{
+			return FCinematicViewportCommands::Get().CustomSafe.Get()->GetInputText();
+		}
+		return GetDisplayName();
+	}
 
 	const FSlateBrush* GetThumbnail() const { return nullptr; }
 
@@ -339,6 +380,7 @@ struct FFilmOverlay_LetterBox : IFilmOverlay
 	}
 
 	FText GetDisplayName() const { return LOCTEXT("LetterboxMask", "Letterbox Mask"); }
+	FText GetToolTip() const { return FCinematicViewportCommands::Get().Letterbox.Get()->GetInputText(); }
 
 	const FSlateBrush* GetThumbnail() const { return nullptr; }
 
@@ -464,10 +506,10 @@ void SFilmOverlayOptions::Construct(const FArguments& InArgs)
 	UFilmOverlayToolkit::RegisterPrimaryFilmOverlay(NAME_None, PrimaryOverlays.Last());
 
 	PrimaryOverlays.Add(MakeShareable(new FFilmOverlay_Grid(3, 3)));
-	UFilmOverlayToolkit::RegisterPrimaryFilmOverlay("3x3Grid", PrimaryOverlays.Last());
+	UFilmOverlayToolkit::RegisterPrimaryFilmOverlay("Grid3x3", PrimaryOverlays.Last());
 
 	PrimaryOverlays.Add(MakeShareable(new FFilmOverlay_Grid(2, 2)));
-	UFilmOverlayToolkit::RegisterPrimaryFilmOverlay("2x2Grid", PrimaryOverlays.Last());
+	UFilmOverlayToolkit::RegisterPrimaryFilmOverlay("Grid2x2", PrimaryOverlays.Last());
 
 	PrimaryOverlays.Add(MakeShareable(new FFilmOverlay_Crosshair));
 	UFilmOverlayToolkit::RegisterPrimaryFilmOverlay("Crosshair", PrimaryOverlays.Last());
@@ -475,14 +517,14 @@ void SFilmOverlayOptions::Construct(const FArguments& InArgs)
 	PrimaryOverlays.Add(MakeShareable(new FFilmOverlay_Rabatment));
 	UFilmOverlayToolkit::RegisterPrimaryFilmOverlay("Rabatment", PrimaryOverlays.Last());
 
-	ToggleableOverlays.Add(MakeShareable(new FFilmOverlay_SafeFrame(LOCTEXT("ActionSafeFrame", "Action Safe"), 95.f, FLinearColor::Red)));
-	UFilmOverlayToolkit::RegisterToggleableFilmOverlay("ActionSafeFrame", ToggleableOverlays.Last());
+	ToggleableOverlays.Add(MakeShareable(new FFilmOverlay_SafeFrame(LOCTEXT("ActionSafe", "Action Safe"), 95.f, FLinearColor::Red)));
+	UFilmOverlayToolkit::RegisterToggleableFilmOverlay("ActionSafe", ToggleableOverlays.Last());
 
-	ToggleableOverlays.Add(MakeShareable(new FFilmOverlay_SafeFrame(LOCTEXT("TitleSafeFrame", "Title Safe"), 90.f, FLinearColor::Yellow)));
-	UFilmOverlayToolkit::RegisterToggleableFilmOverlay("TitleSafeFrame", ToggleableOverlays.Last());
+	ToggleableOverlays.Add(MakeShareable(new FFilmOverlay_SafeFrame(LOCTEXT("TitleSafe", "Title Safe"), 90.f, FLinearColor::Yellow)));
+	UFilmOverlayToolkit::RegisterToggleableFilmOverlay("TitleSafe", ToggleableOverlays.Last());
 
-	ToggleableOverlays.Add(MakeShareable(new FFilmOverlay_SafeFrame(LOCTEXT("CustomSafeFrame", "Custom Safe"), 85.f, FLinearColor::Green)));
-	UFilmOverlayToolkit::RegisterToggleableFilmOverlay("CustomSafeFrame", ToggleableOverlays.Last());
+	ToggleableOverlays.Add(MakeShareable(new FFilmOverlay_SafeFrame(LOCTEXT("CustomSafe", "Custom Safe"), 85.f, FLinearColor::Green)));
+	UFilmOverlayToolkit::RegisterToggleableFilmOverlay("CustomSafe", ToggleableOverlays.Last());
 	
 	ToggleableOverlays.Add(MakeShareable(new FFilmOverlay_LetterBox));
 	UFilmOverlayToolkit::RegisterToggleableFilmOverlay("LetterBox",	ToggleableOverlays.Last());
@@ -612,6 +654,7 @@ TSharedRef<SWidget> SFilmOverlayOptions::ConstructPrimaryOverlaysMenu()
 			SNew(SButton)
 			.ButtonStyle(FAppStyle::Get(), "HoverHintOnly")
 			.OnClicked(this, &SFilmOverlayOptions::SetPrimaryFilmOverlay, OverlayNames[OverlayIndex])
+			.ToolTipText(Overlay.GetToolTip())
 			[
 				SNew(SVerticalBox)
 
@@ -683,6 +726,7 @@ TSharedRef<SWidget> SFilmOverlayOptions::ConstructToggleableOverlaysMenu()
 			SNew(SCheckBox)
 			.OnCheckStateChanged_Lambda(OnCheckStateChanged)
 			.IsChecked_Lambda(IsChecked)
+			.ToolTipText(FilmOverlay->GetToolTip())
 			[
 				SNew(STextBlock)
 				.Text(Pair.Value->GetDisplayName())
@@ -718,6 +762,17 @@ FReply SFilmOverlayOptions::SetPrimaryFilmOverlay(FName InName)
 	return FReply::Unhandled();
 }
 
+FReply SFilmOverlayOptions::ToggleFilmOverlay(FName InName)
+{
+	const TSharedPtr<IFilmOverlay>* FilmOverlay = UFilmOverlayToolkit::GetToggleableFilmOverlays().Find(InName);
+
+	if (FilmOverlay && FilmOverlay->IsValid())
+	{
+		FilmOverlay->Get()->SetEnabled(!FilmOverlay->Get()->IsEnabled());
+	}
+	return FReply::Unhandled();
+}
+
 TSharedRef<SFilmOverlay> SFilmOverlayOptions::GetFilmOverlayWidget() const
 {
 	return OverlayWidget.ToSharedRef();
@@ -768,4 +823,79 @@ IFilmOverlay* SFilmOverlayOptions::GetPrimaryFilmOverlay() const
 	return nullptr;
 }
 
+void SFilmOverlayOptions::BindCommands(TSharedRef<FUICommandList> Bindings)
+{
+	const FCinematicViewportCommands& Commands = FCinematicViewportCommands::Get();
+
+	TArray<FName> PrimaryOverlayNames;
+	UFilmOverlayToolkit::GetPrimaryFilmOverlays().GenerateKeyArray(PrimaryOverlayNames);
+
+	TArray<FName> ToggleableOverlayNames;
+	UFilmOverlayToolkit::GetToggleableFilmOverlays().GenerateKeyArray(ToggleableOverlayNames);
+
+	Bindings->MapAction(
+		Commands.Disabled,
+		FExecuteAction::CreateLambda([this] { SetPrimaryFilmOverlay(NAME_None); }));
+
+	for (int32 OverlayIndex = 0; OverlayIndex < PrimaryOverlayNames.Num(); ++OverlayIndex)
+	{
+		IFilmOverlay& Overlay = *UFilmOverlayToolkit::GetPrimaryFilmOverlays()[PrimaryOverlayNames[OverlayIndex]].Get();
+
+		if (Commands.Grid2x2.Get()->GetCommandName() == PrimaryOverlayNames[OverlayIndex])
+		{
+			Bindings->MapAction(
+				Commands.Grid2x2,
+				FExecuteAction::CreateLambda([this, OverlayName = PrimaryOverlayNames[OverlayIndex]] { SetPrimaryFilmOverlay(OverlayName); }));
+		}
+		else if (Commands.Grid3x3.Get()->GetCommandName() == PrimaryOverlayNames[OverlayIndex])
+		{
+			Bindings->MapAction(
+				Commands.Grid3x3,
+				FExecuteAction::CreateLambda([this, OverlayName = PrimaryOverlayNames[OverlayIndex]] { SetPrimaryFilmOverlay(OverlayName); }));
+		}
+		else if (Commands.Crosshair.Get()->GetCommandName() == PrimaryOverlayNames[OverlayIndex])
+		{
+			Bindings->MapAction(
+				Commands.Crosshair,
+				FExecuteAction::CreateLambda([this, OverlayName = PrimaryOverlayNames[OverlayIndex]] { SetPrimaryFilmOverlay(OverlayName); }));
+		}
+		else if (Commands.Rabatment.Get()->GetCommandName() == PrimaryOverlayNames[OverlayIndex])
+		{
+			Bindings->MapAction(
+				Commands.Rabatment,
+				FExecuteAction::CreateLambda([this, OverlayName = PrimaryOverlayNames[OverlayIndex]] { SetPrimaryFilmOverlay(OverlayName); }));
+		}
+	}
+
+	for (int32 OverlayIndex = 0; OverlayIndex < ToggleableOverlayNames.Num(); ++OverlayIndex)
+	{
+		IFilmOverlay& Overlay = *UFilmOverlayToolkit::GetToggleableFilmOverlays()[ToggleableOverlayNames[OverlayIndex]].Get();
+
+		if (Commands.ActionSafe.Get()->GetCommandName() == ToggleableOverlayNames[OverlayIndex])
+		{
+			Bindings->MapAction(
+				Commands.ActionSafe,
+				FExecuteAction::CreateLambda([this, OverlayName = ToggleableOverlayNames[OverlayIndex]] { ToggleFilmOverlay(OverlayName); }));
+		}
+		else if (Commands.TitleSafe.Get()->GetCommandName() == ToggleableOverlayNames[OverlayIndex])
+		{
+			Bindings->MapAction(
+				Commands.TitleSafe,
+				FExecuteAction::CreateLambda([this, OverlayName = ToggleableOverlayNames[OverlayIndex]] { ToggleFilmOverlay(OverlayName); }));
+		}
+		else if (Commands.CustomSafe.Get()->GetCommandName() == ToggleableOverlayNames[OverlayIndex])
+		{
+			Bindings->MapAction(
+				Commands.CustomSafe,
+				FExecuteAction::CreateLambda([this, OverlayName = ToggleableOverlayNames[OverlayIndex]] { ToggleFilmOverlay(OverlayName); }));
+		}
+		else if (Commands.Letterbox.Get()->GetCommandName() == ToggleableOverlayNames[OverlayIndex])
+		{
+			Bindings->MapAction(
+				Commands.Letterbox,
+				FExecuteAction::CreateLambda([this, OverlayName = ToggleableOverlayNames[OverlayIndex]] { ToggleFilmOverlay(OverlayName); }));
+		}
+	}
+}
+
 #undef LOCTEXT_NAMESPACE
diff --git a/Engine/Plugins/MovieScene/LevelSequenceEditor/Source/LevelSequenceEditor/Private/CinematicViewport/FilmOverlays.h b/Engine/Plugins/MovieScene/LevelSequenceEditor/Source/LevelSequenceEditor/Private/CinematicViewport/FilmOverlays.h
--- a/Engine/Plugins/MovieScene/LevelSequenceEditor/Source/LevelSequenceEditor/Private/CinematicViewport/FilmOverlays.h
+++ b/Engine/Plugins/MovieScene/LevelSequenceEditor/Source/LevelSequenceEditor/Private/CinematicViewport/FilmOverlays.h
@@ -45,6 +45,9 @@ public:
 	/** Retrieve the actual overlay widget that this widget controls. Can be positioned in any other widget hierarchy. */
 	TSharedRef<SFilmOverlay> GetFilmOverlayWidget() const;
 
+	/** Bind commands for the overlays */
+	void BindCommands(TSharedRef<FUICommandList>);
+
 private:
 
 	/** Generate menu content for the combo button */
@@ -73,6 +76,9 @@ private:
 	/** Get/Set the color tint override for the current primary overlay */
 	FLinearColor GetPrimaryColorTint() const;
 	void OnPrimaryColorTintChanged(const FLinearColor& Tint);
+	
+	/** Toggle the film overlay enabled or disabled */
+	FReply ToggleFilmOverlay(FName InName);
 
 private:
 
diff --git a/Engine/Plugins/MovieScene/LevelSequenceEditor/Source/LevelSequenceEditor/Private/CinematicViewport/SCinematicLevelViewport.cpp b/Engine/Plugins/MovieScene/LevelSequenceEditor/Source/LevelSequenceEditor/Private/CinematicViewport/SCinematicLevelViewport.cpp
--- a/Engine/Plugins/MovieScene/LevelSequenceEditor/Source/LevelSequenceEditor/Private/CinematicViewport/SCinematicLevelViewport.cpp
+++ b/Engine/Plugins/MovieScene/LevelSequenceEditor/Source/LevelSequenceEditor/Private/CinematicViewport/SCinematicLevelViewport.cpp
@@ -19,6 +19,7 @@
 #include "Tracks/MovieSceneCinematicShotTrack.h"
 #include "Sections/MovieSceneCinematicShotSection.h"
 #include "SequencerKeyCollection.h"
+#include "CinematicViewport/CinematicViewportCommands.h"
 #include "CinematicViewport/SCinematicTransportRange.h"
 #include "CinematicViewport/FilmOverlays.h"
 #include "Widgets/Layout/SWidgetSwitcher.h"
@@ -448,8 +449,12 @@ void SCinematicLevelViewport::Construct(const FArguments& InArgs)
 	});
 
 	CommandList = MakeShareable( new FUICommandList );
+
 	// Ensure the commands are registered
+	FCinematicViewportCommands::Register();
 	FLevelSequenceEditorCommands::Register();
+
+	FilmOverlayOptions->BindCommands(CommandList.ToSharedRef());
 }
 END_SLATE_FUNCTION_BUILD_OPTIMIZATION
 
diff --git a/Engine/Plugins/MovieScene/LevelSequenceEditor/Source/LevelSequenceEditor/Public/IFilmOverlay.h b/Engine/Plugins/MovieScene/LevelSequenceEditor/Source/LevelSequenceEditor/Public/IFilmOverlay.h
--- a/Engine/Plugins/MovieScene/LevelSequenceEditor/Source/LevelSequenceEditor/Public/IFilmOverlay.h
+++ b/Engine/Plugins/MovieScene/LevelSequenceEditor/Source/LevelSequenceEditor/Public/IFilmOverlay.h
@@ -19,6 +19,9 @@ struct IFilmOverlay
 	/** Get a localized display name that is representative of this overlay */
 	virtual FText GetDisplayName() const = 0;
 
+	/** Get a localized tooltip for this overlay */
+	virtual FText GetToolTip() const { return GetDisplayName(); }
+
 	/** Get a slate thumbnail brush that is representative of this overlay. 36x24 recommended */
 	virtual const FSlateBrush* GetThumbnail() const = 0;
 
