Sequencer: Fix cinematic viewport's transport range going missing when switching sequences

diff --git a/Engine/Plugins/MovieScene/LevelSequenceEditor/Source/LevelSequenceEditor/Private/CinematicViewport/SCinematicLevelViewport.cpp b/Engine/Plugins/MovieScene/LevelSequenceEditor/Source/LevelSequenceEditor/Private/CinematicViewport/SCinematicLevelViewport.cpp
--- a/Engine/Plugins/MovieScene/LevelSequenceEditor/Source/LevelSequenceEditor/Private/CinematicViewport/SCinematicLevelViewport.cpp
+++ b/Engine/Plugins/MovieScene/LevelSequenceEditor/Source/LevelSequenceEditor/Private/CinematicViewport/SCinematicLevelViewport.cpp
@@ -640,14 +640,12 @@ void SCinematicLevelViewport::CleanUp()
 	TransportControlsContainer->SetContent(SNullWidget::NullWidget);
 	TimeRangeContainer->SetContent(SNullWidget::NullWidget);
 
+	CurrentToolkit.Reset();
 }
 
 void SCinematicLevelViewport::OnEditorOpened(FLevelSequenceEditorToolkit& Toolkit)
 {
-	if (!CurrentToolkit.IsValid())
-	{
-		Setup(Toolkit);
-	}
+	Setup(Toolkit);
 }
 
 void SCinematicLevelViewport::OnEditorClosed()
