Sequencer: Mark the spawn register as being cleaned up so that the sequence isn't inadvertently marked dirty

- Tested that making modifications to a spawnable, still invoke OnObjectModified and cause the sequence to be marked dirty.
- Tested that making no modifications to spawnables, closing the sequence, do NOT cause the sequence to be marked dirty.

diff --git a/Engine/Plugins/MovieScene/LevelSequenceEditor/Source/LevelSequenceEditor/Private/Misc/LevelSequenceEditorSpawnRegister.cpp b/Engine/Plugins/MovieScene/LevelSequenceEditor/Source/LevelSequenceEditor/Private/Misc/LevelSequenceEditorSpawnRegister.cpp
--- a/Engine/Plugins/MovieScene/LevelSequenceEditor/Source/LevelSequenceEditor/Private/Misc/LevelSequenceEditorSpawnRegister.cpp
+++ b/Engine/Plugins/MovieScene/LevelSequenceEditor/Source/LevelSequenceEditor/Private/Misc/LevelSequenceEditorSpawnRegister.cpp
@@ -255,7 +255,12 @@ void FLevelSequenceEditorSpawnRegister::OnObjectModified(UObject* ModifiedObject
 		ModifiedObject = ModifiedObject->GetOuter();
 	}
 
-	if (TrackedState)
+	// 
+	// When the register is cleaning itself up, there's no need to mark the sequence package dirty because spawnables are just getting removed.
+	// This scenario can happen if SpawnableA is attached to SpawnableB. When Restore State occurs, SpawnableB is destroyed, which causes 
+	// SpawnableA to be marked as modified, resulting in marking the sequencer package dirty without this check. 
+	//
+	if (TrackedState && !bCleaningUp)
 	{
 		TrackedState->bHasBeenModified = true;
 
diff --git a/Engine/Source/Runtime/MovieScene/Private/MovieSceneSpawnRegister.cpp b/Engine/Source/Runtime/MovieScene/Private/MovieSceneSpawnRegister.cpp
--- a/Engine/Source/Runtime/MovieScene/Private/MovieSceneSpawnRegister.cpp
+++ b/Engine/Source/Runtime/MovieScene/Private/MovieSceneSpawnRegister.cpp
@@ -103,6 +103,8 @@ void FMovieSceneSpawnRegister::ForgetExternallyOwnedSpawnedObjects(FMovieSceneEv
 
 void FMovieSceneSpawnRegister::CleanUp(IMovieScenePlayer& Player)
 {
+	TGuardValue<bool> CleaningUp(bCleaningUp, true);
+
 	DestroyObjectsByPredicate(Player, [&](const FGuid&, ESpawnOwnership, FMovieSceneSequenceIDRef){
 		return true;
 	});
diff --git a/Engine/Source/Runtime/MovieScene/Public/MovieSceneSpawnRegister.h b/Engine/Source/Runtime/MovieScene/Public/MovieSceneSpawnRegister.h
--- a/Engine/Source/Runtime/MovieScene/Public/MovieSceneSpawnRegister.h
+++ b/Engine/Source/Runtime/MovieScene/Public/MovieSceneSpawnRegister.h
@@ -254,6 +254,9 @@ protected:
 
 	/** Register of spawned objects */
 	TMap<FMovieSceneSpawnRegisterKey, FSpawnedObject> Register;
+
+	/** True when cleaning ourselves up. Used to bypass marking a sequence dirty when objects are modified since we're cleaning ourselves up */
+	bool bCleaningUp = false;
 };
 
 class FNullMovieSceneSpawnRegister : public FMovieSceneSpawnRegister
