Sequencer: Fix Snap Sections to Timeline using Source Timecode so that it takes tick resolution into account

diff --git a/Engine/Plugins/MovieScene/LevelSequenceEditor/Source/LevelSequenceEditor/Private/LevelSequenceEditorSubsystem.cpp b/Engine/Plugins/MovieScene/LevelSequenceEditor/Source/LevelSequenceEditor/Private/LevelSequenceEditorSubsystem.cpp
--- a/Engine/Plugins/MovieScene/LevelSequenceEditor/Source/LevelSequenceEditor/Private/LevelSequenceEditorSubsystem.cpp
+++ b/Engine/Plugins/MovieScene/LevelSequenceEditor/Source/LevelSequenceEditor/Private/LevelSequenceEditorSubsystem.cpp
@@ -490,6 +490,7 @@ void ULevelSequenceEditorSubsystem::SnapSectionsToTimelineUsingSourceTimecode(co
 	bool bAnythingChanged = false;
 
 	const FFrameRate TickResolution = Sequencer->GetFocusedTickResolution();
+	const FFrameRate DisplayRate = Sequencer->GetFocusedDisplayRate();
 
 	for (UMovieSceneSection* Section : Sections)
 	{
@@ -505,7 +506,7 @@ void ULevelSequenceEditorSubsystem::SnapSectionsToTimelineUsingSourceTimecode(co
 			continue;
 		}
 
-		const FFrameNumber SectionSourceStartFrameNumber = SectionSourceTimecode.ToFrameNumber(TickResolution);
+		const FFrameNumber SectionSourceStartFrameNumber = SectionSourceTimecode.ToFrameNumber(DisplayRate);
 
 		// Account for any trimming at the start of the section when computing the
 		// target frame number to move this section to.
@@ -514,7 +515,7 @@ void ULevelSequenceEditorSubsystem::SnapSectionsToTimelineUsingSourceTimecode(co
 
 		const FFrameNumber SectionCurrentStartFrameNumber = Section->GetInclusiveStartFrame();
 
-		const FFrameNumber Delta = -(SectionCurrentStartFrameNumber - TargetFrameNumber);
+		const FFrameNumber Delta = -(SectionCurrentStartFrameNumber - ConvertFrameTime(TargetFrameNumber, DisplayRate, TickResolution).GetFrame().Value);
 
 		Section->MoveSection(Delta);
 
@@ -597,20 +598,22 @@ void ULevelSequenceEditorSubsystem::SyncSectionsUsingSourceTimecode(const TArray
 	bool bAnythingChanged = false;
 
 	const FFrameRate TickResolution = Sequencer->GetFocusedTickResolution();
+	const FFrameRate DisplayRate = Sequencer->GetFocusedDisplayRate();
 
 	const UMovieSceneSection* FirstSection = SectionsToSync[0];
-	const FFrameNumber FirstSectionSourceTimecode = FirstSection->TimecodeSource.Timecode.ToFrameNumber(TickResolution);
+	const FFrameNumber FirstSectionSourceTimecode = FirstSection->TimecodeSource.Timecode.ToFrameNumber(DisplayRate);
+
 	const FFrameNumber FirstSectionCurrentStartFrame = FirstSection->GetInclusiveStartFrame();
 	const FFrameNumber FirstSectionOffsetFrames = FirstSection->GetOffsetTime().Get(FFrameTime()).FloorToFrame();
 	SectionsToSync.RemoveAt(0);
 
 	for (UMovieSceneSection* Section : SectionsToSync)
 	{
-		const FFrameNumber SectionSourceTimecode = Section->TimecodeSource.Timecode.ToFrameNumber(TickResolution);
+		const FFrameNumber SectionSourceTimecode = Section->TimecodeSource.Timecode.ToFrameNumber(DisplayRate);
 		const FFrameNumber SectionCurrentStartFrame = Section->GetInclusiveStartFrame();
 		const FFrameNumber SectionOffsetFrames = Section->GetOffsetTime().Get(FFrameTime()).FloorToFrame();
 
-		const FFrameNumber TimecodeDelta = SectionSourceTimecode - FirstSectionSourceTimecode;
+		const FFrameNumber TimecodeDelta = ConvertFrameTime(SectionSourceTimecode - FirstSectionSourceTimecode, DisplayRate, TickResolution).GetFrame().Value;
 		const FFrameNumber CurrentDelta = (SectionCurrentStartFrame - SectionOffsetFrames) - (FirstSectionCurrentStartFrame - FirstSectionOffsetFrames);
 		const FFrameNumber Delta = -CurrentDelta + TimecodeDelta;
 
