Sequencer: Initialize layers for spawned actors. This fixes issues where if you add a spawnable to a layer, hide the layer, and then reopen the sequence in a new level, the spawnable won't be visible and won't respond to layers in the new level.

diff --git a/Engine/Plugins/MovieScene/LevelSequenceEditor/Source/LevelSequenceEditor/Private/Misc/LevelSequenceEditorSpawnRegister.cpp b/Engine/Plugins/MovieScene/LevelSequenceEditor/Source/LevelSequenceEditor/Private/Misc/LevelSequenceEditorSpawnRegister.cpp
--- a/Engine/Plugins/MovieScene/LevelSequenceEditor/Source/LevelSequenceEditor/Private/Misc/LevelSequenceEditorSpawnRegister.cpp
+++ b/Engine/Plugins/MovieScene/LevelSequenceEditor/Source/LevelSequenceEditor/Private/Misc/LevelSequenceEditorSpawnRegister.cpp
@@ -4,6 +4,7 @@
 #include "Modules/ModuleManager.h"
 #include "ISequencer.h"
 #include "LevelEditor.h"
+#include "Layers/LayersSubsystem.h"
 #include "EntitySystem/MovieSceneSpawnablesSystem.h"
 #include "Evaluation/MovieSceneEvaluationTemplateInstance.h"
 #include "MovieScene.h"
@@ -72,6 +73,11 @@ UObject* FLevelSequenceEditorSpawnRegister::SpawnObject(FMovieSceneSpawnable& Sp
 		// Add an entry to the tracked objects map to keep track of this object (so that it can be saved when modified)
 		TrackedObjects.Add(NewActor, FTrackedObjectState(TemplateID, Spawnable.GetGuid()));
 
+		if (ULayersSubsystem* Layers = GEditor->GetEditorSubsystem<ULayersSubsystem>())
+		{
+			Layers->InitializeNewActorLayers(NewActor);
+		}
+
 		// Select the actor if we think it should be selected
 		if (SelectedSpawnedObjects.Contains(FMovieSceneSpawnRegisterKey(TemplateID, Spawnable.GetGuid())))
 		{
