Sequencer: Clamp shot duration to max frame rate

diff --git a/Engine/Plugins/MovieScene/LevelSequenceEditor/Source/LevelSequenceEditor/Private/LevelSequenceEditorToolkit.cpp b/Engine/Plugins/MovieScene/LevelSequenceEditor/Source/LevelSequenceEditor/Private/LevelSequenceEditorToolkit.cpp
--- a/Engine/Plugins/MovieScene/LevelSequenceEditor/Source/LevelSequenceEditor/Private/LevelSequenceEditorToolkit.cpp
+++ b/Engine/Plugins/MovieScene/LevelSequenceEditor/Source/LevelSequenceEditor/Private/LevelSequenceEditorToolkit.cpp
@@ -9,6 +9,7 @@
 #include "Containers/ArrayBuilder.h"
 #include "KeyParams.h"
 #include "LevelSequenceEditorModule.h"
+#include "LevelSequenceEditorSubsystem.h"
 #include "Misc/LevelSequenceEditorSettings.h"
 #include "Misc/LevelSequenceEditorSpawnRegister.h"
 #include "Misc/LevelSequenceEditorHelpers.h"
@@ -835,7 +836,16 @@ void FLevelSequenceEditorToolkit::HandleLevelSequenceWithShotsCreated(UObject* L
 	FFrameNumber SequenceStartTime = (ProjectSettings->DefaultStartTime * TickResolution).FloorToFrame();
 	FFrameNumber ShotStartTime = SequenceStartTime;
 	FFrameNumber ShotEndTime   = ShotStartTime;
+
 	int32        ShotDuration  = (ProjectSettings->DefaultDuration * TickResolution).RoundToFrame().Value;
+	
+	if (ProjectSettings->DefaultDuration * NumShots > TickResolution.MaxSeconds())
+	{
+		ShotDuration = ((TickResolution.MaxSeconds() / NumShots) * TickResolution).RoundToFrame().Value;
+
+		UE_LOG(LogLevelSequenceEditor, Warning, TEXT("Default shot duration: %f is too large for creating %d shots. Clamping shot duration to: %f"), ProjectSettings->DefaultDuration, NumShots, TickResolution.MaxSeconds() / NumShots);
+	}
+	
 	FString FirstShotName; 
 	for (uint32 ShotIndex = 0; ShotIndex < NumShots; ++ShotIndex)
 	{
@@ -886,7 +896,6 @@ void FLevelSequenceEditorToolkit::HandleLevelSequenceWithShotsCreated(UObject* L
 	}
 }
 
-
 TSharedRef<FExtender> FLevelSequenceEditorToolkit::HandleMenuExtensibilityGetExtender(const TSharedRef<FUICommandList> CommandList, const TArray<UObject*> ContextSensitiveObjects)
 {
 	TSharedRef<FExtender> AddTrackMenuExtender(new FExtender());
