Sequencer: Check for valid level sequence actor

#rnx

diff --git a/Engine/Plugins/MovieScene/LevelSequenceEditor/Source/LevelSequenceEditor/Private/LevelSequenceEditorToolkit.cpp b/Engine/Plugins/MovieScene/LevelSequenceEditor/Source/LevelSequenceEditor/Private/LevelSequenceEditorToolkit.cpp
--- a/Engine/Plugins/MovieScene/LevelSequenceEditor/Source/LevelSequenceEditor/Private/LevelSequenceEditorToolkit.cpp
+++ b/Engine/Plugins/MovieScene/LevelSequenceEditor/Source/LevelSequenceEditor/Private/LevelSequenceEditorToolkit.cpp
@@ -885,14 +885,16 @@ void FLevelSequenceEditorToolkit::HandleLevelSequenceWithShotsCreated(UObject* L
 		return;
 	}
 
-	ALevelSequenceActor* NewActor = CastChecked<ALevelSequenceActor>(GEditor->UseActorFactory(ActorFactory, FAssetData(LevelSequenceWithShotsAsset), &FTransform::Identity));
-	if (GCurrentLevelEditingViewportClient != nullptr && GCurrentLevelEditingViewportClient->IsPerspective())
+	if (ALevelSequenceActor* NewActor = Cast<ALevelSequenceActor>(GEditor->UseActorFactory(ActorFactory, FAssetData(LevelSequenceWithShotsAsset), &FTransform::Identity)))
 	{
-		GEditor->MoveActorInFrontOfCamera(*NewActor, GCurrentLevelEditingViewportClient->GetViewLocation(), GCurrentLevelEditingViewportClient->GetViewRotation().Vector());
-	}
-	else
-	{
-		GEditor->MoveViewportCamerasToActor(*NewActor, false);
+		if (GCurrentLevelEditingViewportClient != nullptr && GCurrentLevelEditingViewportClient->IsPerspective())
+		{
+			GEditor->MoveActorInFrontOfCamera(*NewActor, GCurrentLevelEditingViewportClient->GetViewLocation(), GCurrentLevelEditingViewportClient->GetViewRotation().Vector());
+		}
+		else
+		{
+			GEditor->MoveViewportCamerasToActor(*NewActor, false);
+		}
 	}
 }
 
