Sequencer: Stop the sequence from being marked dirty (because of the spawnable being modified) if the sequence is evaluating.

diff --git a/Engine/Plugins/MovieScene/LevelSequenceEditor/Source/LevelSequenceEditor/Private/Misc/LevelSequenceEditorSpawnRegister.cpp b/Engine/Plugins/MovieScene/LevelSequenceEditor/Source/LevelSequenceEditor/Private/Misc/LevelSequenceEditorSpawnRegister.cpp
--- a/Engine/Plugins/MovieScene/LevelSequenceEditor/Source/LevelSequenceEditor/Private/Misc/LevelSequenceEditorSpawnRegister.cpp
+++ b/Engine/Plugins/MovieScene/LevelSequenceEditor/Source/LevelSequenceEditor/Private/Misc/LevelSequenceEditorSpawnRegister.cpp
@@ -254,6 +254,15 @@ void FLevelSequenceEditorSpawnRegister::OnObjectsReplaced(const TMap<UObject*, U
 
 void FLevelSequenceEditorSpawnRegister::OnObjectModified(UObject* ModifiedObject)
 {
+	// If the sequence is evaluating, we don't want object modifications to dirty the sequence itself. 
+	// For example, this protects against situations where OnObjectModified would be called in response 
+	// to the spawnable being attached with AttachToComponent
+	TSharedPtr<ISequencer> Sequencer = WeakSequencer.Pin();
+	if (!Sequencer.IsValid() || Sequencer->IsEvaluating())
+	{
+		return;
+	}
+
 	FTrackedObjectState* TrackedState = TrackedObjects.Find(ModifiedObject);
 	while (!TrackedState && ModifiedObject)
 	{
@@ -270,8 +279,7 @@ void FLevelSequenceEditorSpawnRegister::OnObjectModified(UObject* ModifiedObject
 	{
 		TrackedState->bHasBeenModified = true;
 
-		TSharedPtr<ISequencer> Sequencer = WeakSequencer.Pin();
-		UMovieSceneSequence*   OwningSequence = Sequencer.IsValid() ? Sequencer->GetEvaluationTemplate().GetSequence(TrackedState->TemplateID) : nullptr;
+		UMovieSceneSequence*   OwningSequence = Sequencer->GetEvaluationTemplate().GetSequence(TrackedState->TemplateID);
 		if (OwningSequence)
 		{
 			OwningSequence->MarkPackageDirty();
