Sequencer: Add per key options to smart bake transforms. Use similar settings and dialog to the the other smart baking setups. Main thing is we need to calculate the keys, and to do that we need to get the correct parents,  including from attachments.

diff --git a/Engine/Plugins/Animation/ControlRig/Source/ControlRigEditor/Private/Tools/ControlRigSnapper.cpp b/Engine/Plugins/Animation/ControlRig/Source/ControlRigEditor/Private/Tools/ControlRigSnapper.cpp
--- a/Engine/Plugins/Animation/ControlRig/Source/ControlRigEditor/Private/Tools/ControlRigSnapper.cpp
+++ b/Engine/Plugins/Animation/ControlRig/Source/ControlRigEditor/Private/Tools/ControlRigSnapper.cpp
@@ -326,31 +326,6 @@ static bool CalculateWorldTransformsFromParents(ISequencer* Sequencer, const FCo
 	return false;
 }
 
-//to do handle constraints and spaces, but for now we just output a single actor parent.
-static void GetActorParents(const FActorForWorldTransforms& ActorSelection,
-	TArray<FActorForWorldTransforms>& OutParentActors, TArray<FControlRigForWorldTransforms>& OutControlRigs)
-{
-	if (ActorSelection.Actor.IsValid())
-	{
-		//has component so parent is next component
-		if (ActorSelection.Component.IsValid() && ActorSelection.Component->GetAttachParent())
-		{
-			FActorForWorldTransforms OutParentActor;
-			OutParentActor.Actor = ActorSelection.Actor;
-			OutParentActor.Component = ActorSelection.Component->GetAttachParent();
-			OutParentActors.Add(OutParentActor);
-			GetActorParents(OutParentActor, OutParentActors, OutControlRigs);
-		}
-		else if (AActor* ParentActor = ActorSelection.Actor->GetAttachParentActor())
-		{
-			FActorForWorldTransforms OutParentActor;
-			OutParentActor.Actor = ParentActor;
-			OutParentActors.Add(OutParentActor);
-			GetActorParents(OutParentActor, OutParentActors, OutControlRigs);
-		}
-	}
-}
-
 //todo handle constraints/spaces
 static void GetControlRigParents(const FControlRigForWorldTransforms& ControlRigAndSelection, 
 	TArray<FActorForWorldTransforms>& OutParentActors, TArray<FControlRigForWorldTransforms>& OutControlRigs)
@@ -388,7 +363,7 @@ static void GetControlRigParents(const FControlRigForWorldTransforms& ControlRig
 						FActorForWorldTransforms OutParentActor;
 						OutParentActor.Actor = Actor;
 						OutParentActors.Add(OutParentActor);
-						GetActorParents(OutParentActor, OutParentActors, OutControlRigs);
+						MovieSceneToolHelpers::GetActorParents(OutParentActor, OutParentActors);
 					}
 				}
 			}
@@ -437,8 +412,7 @@ static void GetTransformFrames(ISequencer * Sequencer, const FControlRigSnapperS
 		ParentActors.Reset();
 		ParentControlRigs.Reset();
 		ParentActors.Add(ActorsFor);
-		GetActorParents(ActorsFor, ParentActors, ParentControlRigs);
-
+		MovieSceneToolHelpers::GetActorParents(ActorsFor, ParentActors);
 	}
 	for (const FControlRigForWorldTransforms& ControlRigsFor : SnapItem.ControlRigs)
 	{
diff --git a/Engine/Plugins/MovieScene/LevelSequenceEditor/Source/LevelSequenceEditor/Private/LevelSequenceEditorSubsystem.cpp b/Engine/Plugins/MovieScene/LevelSequenceEditor/Source/LevelSequenceEditor/Private/LevelSequenceEditorSubsystem.cpp
--- a/Engine/Plugins/MovieScene/LevelSequenceEditor/Source/LevelSequenceEditor/Private/LevelSequenceEditorSubsystem.cpp
+++ b/Engine/Plugins/MovieScene/LevelSequenceEditor/Source/LevelSequenceEditor/Private/LevelSequenceEditorSubsystem.cpp
@@ -34,6 +34,16 @@
 #include "ScopedTransaction.h"
 #include "Widgets/Notifications/SNotificationList.h"
 #include "Framework/MultiBox/MultiBoxBuilder.h"
+#include "BakingAnimationKeySettings.h"
+#include "IStructureDetailsView.h"
+#include "Widgets/SWidget.h"
+#include "Widgets/SCompoundWidget.h"
+#include "Widgets/Layout/SBox.h"
+#include "Widgets/Input/NumericTypeInterface.h"
+#include "FrameNumberDetailsCustomization.h"
+#include "MovieSceneToolHelpers.h"
+#include "ActorForWorldTransforms.h"
+#include "KeyParams.h"
 
 #include UE_INLINE_GENERATED_CPP_BY_NAME(LevelSequenceEditorSubsystem)
 
@@ -640,6 +650,188 @@ void ULevelSequenceEditorSubsystem::SyncSectionsUsingSourceTimecode(const TArray
 	}
 }
 
+//////////////////////////////////////////////////////////////
+/// SBakeTransformWidget
+///////////////////////////////////////////////////////////
+
+DECLARE_DELEGATE_RetVal_OneParam(FReply, SBakeTransformOnBake, FBakingAnimationKeySettings);
+
+/** Widget allowing baking controls from one space to another */
+class  SBakeTransformWidget : public SCompoundWidget
+{
+public:
+
+	SLATE_BEGIN_ARGS(SBakeTransformWidget)
+		: _Sequencer(nullptr)
+	{}
+		SLATE_ARGUMENT(ISequencer*, Sequencer)
+		SLATE_ARGUMENT(FBakingAnimationKeySettings, Settings)
+		SLATE_EVENT(SBakeTransformOnBake, OnBake)
+		SLATE_END_ARGS()
+
+	void Construct(const FArguments& InArgs);
+	virtual ~SBakeTransformWidget() override {}
+
+	FReply OpenDialog(bool bModal = true);
+	void CloseDialog();
+
+private:
+
+	//used for setting up the details
+	TSharedPtr<TStructOnScope<FBakingAnimationKeySettings>> Settings;
+
+	ISequencer* Sequencer;
+
+	TWeakPtr<SWindow> DialogWindow;
+	TSharedPtr<IStructureDetailsView> DetailsView;
+};
+
+void SBakeTransformWidget::Construct(const FArguments& InArgs)
+{
+	check(InArgs._Sequencer);
+	check(InArgs._OnBake.IsBound());
+
+	Settings = MakeShared<TStructOnScope<FBakingAnimationKeySettings>>();
+	Settings->InitializeAs<FBakingAnimationKeySettings>();
+	*Settings = InArgs._Settings;
+	//always setting space to be parent as default, since stored space may not be available.
+	Sequencer = InArgs._Sequencer;
+
+	FStructureDetailsViewArgs StructureViewArgs;
+	StructureViewArgs.bShowObjects = true;
+	StructureViewArgs.bShowAssets = true;
+	StructureViewArgs.bShowClasses = true;
+	StructureViewArgs.bShowInterfaces = true;
+
+	FDetailsViewArgs ViewArgs;
+	ViewArgs.bAllowSearch = false;
+	ViewArgs.bHideSelectionTip = false;
+	ViewArgs.bShowObjectLabel = false;
+
+	FPropertyEditorModule& PropertyEditor = FModuleManager::Get().LoadModuleChecked<FPropertyEditorModule>(TEXT("PropertyEditor"));
+
+	DetailsView = PropertyEditor.CreateStructureDetailView(ViewArgs, StructureViewArgs, TSharedPtr<FStructOnScope>());
+	TSharedPtr<INumericTypeInterface<double>> NumericTypeInterface = Sequencer->GetNumericTypeInterface();
+	DetailsView->GetDetailsView()->RegisterInstancedCustomPropertyTypeLayout("FrameNumber",
+		FOnGetPropertyTypeCustomizationInstance::CreateLambda([=]() {return MakeShared<FFrameNumberDetailsCustomization>(NumericTypeInterface); }));
+	DetailsView->SetStructureData(Settings);
+
+	ChildSlot
+		[
+			SNew(SBorder)
+			.Visibility(EVisibility::Visible)
+		[
+			SNew(SVerticalBox)
+
+			
+			+ SVerticalBox::Slot()
+				.AutoHeight()
+				.Padding(0.f, 0.f, 0.f, 0.f)
+				[
+					DetailsView->GetWidget().ToSharedRef()
+				]
+
+			+ SVerticalBox::Slot()
+			.AutoHeight()
+			.Padding(0.f, 16.f, 0.f, 0.f)
+			[
+				SNew(SHorizontalBox)
+
+				+ SHorizontalBox::Slot()
+			.FillWidth(1.f)
+			[
+				SNew(SSpacer)
+			]
+
+			+ SHorizontalBox::Slot()
+			.AutoWidth()
+			.HAlign(HAlign_Left)
+			.VAlign(VAlign_Center)
+			.Padding(8.f, 0.f, 0.f, 0.f)
+			[
+				SNew(SButton)
+				.HAlign(HAlign_Center)
+			.ContentPadding(FAppStyle::GetMargin("StandardDialog.ContentPadding"))
+			.Text(LOCTEXT("OK", "OK"))
+			.OnClicked_Lambda([this, InArgs]()
+				{
+					FReply Reply = InArgs._OnBake.Execute( *(Settings->Get()));
+					CloseDialog();
+					return Reply;
+
+				})
+			]
+
+		+ SHorizontalBox::Slot()
+			.AutoWidth()
+			.HAlign(HAlign_Left)
+			.VAlign(VAlign_Center)
+			.Padding(8.f, 0.f, 16.f, 0.f)
+			[
+				SNew(SButton)
+				.HAlign(HAlign_Center)
+				.ContentPadding(FAppStyle::GetMargin("StandardDialog.ContentPadding"))
+				.Text(LOCTEXT("Cancel", "Cancel"))
+				.OnClicked_Lambda([this]()
+				{
+					CloseDialog();
+					return FReply::Handled();
+				})
+			]
+		]
+	]
+	];
+}
+
+class SBakeTransformDialogWindow : public SWindow
+{
+};
+
+FReply SBakeTransformWidget::OpenDialog(bool bModal)
+{
+	check(!DialogWindow.IsValid());
+
+	const FVector2D CursorPos = FSlateApplication::Get().GetCursorPos();
+
+	TSharedRef<SBakeTransformDialogWindow> Window = SNew(SBakeTransformDialogWindow)
+		.Title(LOCTEXT("SBakeTransformWidgetTitle", "Bake Transforms"))
+		.CreateTitleBar(true)
+		.Type(EWindowType::Normal)
+		.SizingRule(ESizingRule::Autosized)
+		.ScreenPosition(CursorPos)
+		.FocusWhenFirstShown(true)
+		.ActivationPolicy(EWindowActivationPolicy::FirstShown)
+		[
+			AsShared()
+		];
+
+	Window->SetWidgetToFocusOnActivate(AsShared());
+
+	DialogWindow = Window;
+
+	Window->MoveWindowTo(CursorPos);
+
+	if (bModal)
+	{
+		GEditor->EditorAddModalWindow(Window);
+	}
+	else
+	{
+		FSlateApplication::Get().AddWindow(Window);
+	}
+
+	return FReply::Handled();
+}
+
+void SBakeTransformWidget::CloseDialog()
+{
+	if (DialogWindow.IsValid())
+	{
+		DialogWindow.Pin()->RequestDestroyWindow();
+		DialogWindow.Reset();
+	}
+}
+
 void ULevelSequenceEditorSubsystem::BakeTransformInternal()
 {
 	TSharedPtr<ISequencer> Sequencer = GetActiveSequencer();
@@ -657,9 +849,6 @@ void ULevelSequenceEditorSubsystem::BakeTransformInternal()
 	FFrameRate TickResolution = FocusedMovieScene->GetTickResolution();
 	FFrameRate DisplayRate = FocusedMovieScene->GetDisplayRate();
 
-	FFrameNumber InFrame = UE::MovieScene::DiscreteInclusiveLower(FocusedMovieScene->GetPlaybackRange());
-	FFrameNumber OutFrame = UE::MovieScene::DiscreteExclusiveUpper(FocusedMovieScene->GetPlaybackRange());
-
 	TArray<FGuid> ObjectBindings;
 	Sequencer->GetSelectedObjects(ObjectBindings);
 
@@ -669,9 +858,23 @@ void ULevelSequenceEditorSubsystem::BakeTransformInternal()
 		BindingProxies.Add(FMovieSceneBindingProxy(Guid, Sequencer->GetFocusedMovieSceneSequence()));
 	}
 
-	FMovieSceneScriptingParams Params;
-	Params.TimeUnit = ESequenceTimeUnit::TickResolution;
-	BakeTransform(BindingProxies, FFrameTime(InFrame), FFrameTime(OutFrame), ConvertFrameTime(1, DisplayRate, TickResolution), Params);
+	static FBakingAnimationKeySettings Settings; //reuse the settings except for the range
+	Settings.StartFrame = UE::MovieScene::DiscreteInclusiveLower(FocusedMovieScene->GetPlaybackRange());
+	Settings.EndFrame = UE::MovieScene::DiscreteExclusiveUpper(FocusedMovieScene->GetPlaybackRange());
+
+	TSharedRef<SBakeTransformWidget> BakeWidget =
+		SNew(SBakeTransformWidget)
+		.Settings(Settings)
+		.Sequencer(Sequencer.Get())
+		.OnBake_Lambda([this, &BindingProxies, Sequencer](FBakingAnimationKeySettings InSettings)
+			{
+				FMovieSceneScriptingParams Params;
+				Params.TimeUnit = ESequenceTimeUnit::TickResolution;
+				BakeTransformWithSettings(BindingProxies, InSettings, Params);
+				return FReply::Handled();
+			});
+
+	BakeWidget->OpenDialog(true);
 }
 
 void ULevelSequenceEditorSubsystem::BakeTransform(const TArray<FMovieSceneBindingProxy>& ObjectBindings, const FFrameTime& BakeInTime, const FFrameTime& BakeOutTime, const FFrameTime& BakeInterval, const FMovieSceneScriptingParams& Params)
@@ -699,15 +902,8 @@ void ULevelSequenceEditorSubsystem::BakeTransform(const TArray<FMovieSceneBindin
 		return;
 	}
 
-	FScopedTransaction BakeTransform(LOCTEXT("BakeTransform", "Bake Transform"));
-
-	FocusedMovieScene->Modify();
-
-	FQualifiedFrameTime ResetTime = Sequencer->GetLocalTime();
-
 	FFrameRate TickResolution = FocusedMovieScene->GetTickResolution();
 	FFrameRate DisplayRate = FocusedMovieScene->GetDisplayRate();
-
 	FFrameTime InFrame = BakeInTime;
 	FFrameTime OutFrame = BakeOutTime;
 	FFrameTime Interval = BakeInterval;
@@ -719,25 +915,129 @@ void ULevelSequenceEditorSubsystem::BakeTransform(const TArray<FMovieSceneBindin
 		Interval = ConvertFrameTime(BakeInterval, DisplayRate, TickResolution);
 	}
 
-	struct FBakeData
+	FBakingAnimationKeySettings Settings;
+	FFrameTime OneFrame = ConvertFrameTime(1, DisplayRate, TickResolution);
+	//we never want subframes when baking so use frame increments
+	Settings.FrameIncrement = Interval.FrameNumber.Value / OneFrame.FrameNumber.Value;
+	Settings.FrameIncrement = Settings.FrameIncrement <= 0 ? 1 : Settings.FrameIncrement;
+	Settings.StartFrame = InFrame.GetFrame();
+	Settings.EndFrame = OutFrame.GetFrame();
+	Settings.BakingKeySettings = EBakingKeySettings::AllFrames;
+	FMovieSceneScriptingParams NewParams;
+	NewParams.TimeUnit = ESequenceTimeUnit::TickResolution;
+	BakeTransformWithSettings(ObjectBindings, Settings, NewParams);
+}
+
+void ULevelSequenceEditorSubsystem::CalculateFramesPerGuid(TSharedPtr<ISequencer>& Sequencer, const FBakingAnimationKeySettings& InSettings, TMap<FGuid, ULevelSequenceEditorSubsystem::FBakeData>& OutBakeDataMap,
+	TSortedMap<FFrameNumber, FFrameNumber>& OutFrameMap)
+{
+	OutFrameMap.Reset();
+	TArray<FFrameNumber> Frames;
+	//we get all frames since we need to get the Actor PER FRAME in order to handle spanwables
+	MovieSceneToolHelpers::CalculateFramesBetween(Sequencer->GetFocusedMovieSceneSequence()->GetMovieScene(), InSettings.StartFrame, InSettings.EndFrame, InSettings.FrameIncrement, Frames);
+	if (InSettings.BakingKeySettings == EBakingKeySettings::AllFrames)
 	{
-		TArray<FVector> Locations;
-		TArray<FRotator> Rotations;
-		TArray<FVector> Scales;
-		TArray<FFrameNumber> KeyTimes;
-	};
+		for (FFrameNumber& Frame : Frames)
+		{
+			OutFrameMap.Add(Frame,Frame);
+		}
+		for (TPair<FGuid, FBakeData>& BakeData : OutBakeDataMap)
+		{
+			BakeData.Value.KeyTimes.Reset();
+			BakeData.Value.KeyTimes = OutFrameMap;
+		}
+	}
+	else
+	{
+		for (TPair<FGuid, FBakeData>& BakeData : OutBakeDataMap)
+		{
+			FActorForWorldTransforms ActorForWorldTransforms;
+			FGuid Guid = BakeData.Key;
 
+			for (TWeakObjectPtr<> RuntimeObject : Sequencer->FindObjectsInCurrentSequence(Guid))
+			{
+				UActorComponent* ActorComponent = nullptr;
+				AActor* Actor = Cast<AActor>(RuntimeObject.Get());
+				if (!Actor)
+				{
+					ActorComponent = Cast<UActorComponent>(RuntimeObject.Get());
+					if (ActorComponent)
+					{
+						Actor = ActorComponent->GetOwner();
+					}
+				}
+				ActorForWorldTransforms.Actor = Actor;
+				ActorForWorldTransforms.Component = Cast<USceneComponent>(ActorComponent);
+				BakeData.Value.KeyTimes.Reset();
+				MovieSceneToolHelpers::GetActorsAndParentsKeyFrames(Sequencer.Get(), ActorForWorldTransforms,
+					InSettings.StartFrame, InSettings.EndFrame, BakeData.Value.KeyTimes);
+				OutFrameMap.Append(BakeData.Value.KeyTimes);
+			}
+		}
+	}
+}
+
+bool ULevelSequenceEditorSubsystem::BakeTransformWithSettings(const TArray<FMovieSceneBindingProxy>& ObjectBindings, const FBakingAnimationKeySettings& InSettings, const FMovieSceneScriptingParams& Params)
+{
+	TSharedPtr<ISequencer> Sequencer = GetActiveSequencer();
+	if (Sequencer == nullptr)
+	{
+		UE_LOG(LogLevelSequenceEditor, Warning, TEXT("Bake Transform failed."));
+		return false;
+	}
+
+	UMovieScene* FocusedMovieScene = Sequencer->GetFocusedMovieSceneSequence()->GetMovieScene();
+	if (!FocusedMovieScene)
+	{
+		UE_LOG(LogLevelSequenceEditor, Warning, TEXT("Bake Transform failed."));
+		return false;
+	}
+
+	if (FocusedMovieScene->IsReadOnly())
+	{
+		UE_LOG(LogLevelSequenceEditor, Warning, TEXT("Bake Transform failed."));
+		FSequencerUtilities::ShowReadOnlyError();
+		return false;
+	}
+
+	if (ObjectBindings.Num() == 0)
+	{
+		UE_LOG(LogLevelSequenceEditor, Warning, TEXT("Bake Transform failed."));
+		return false;
+	}
+
+	FScopedTransaction BakeTransform(LOCTEXT("BakeTransform", "Bake Transform"));
+
+	FocusedMovieScene->Modify();
+
+	FQualifiedFrameTime ResetTime = Sequencer->GetLocalTime();
+
+	FFrameRate TickResolution = FocusedMovieScene->GetTickResolution();
+	FFrameRate DisplayRate = FocusedMovieScene->GetDisplayRate();
+
+	FBakingAnimationKeySettings SettingsInTick = InSettings;
+	
+	if (Params.TimeUnit == ESequenceTimeUnit::DisplayRate)
+	{
+		SettingsInTick.StartFrame = ConvertFrameTime(SettingsInTick.StartFrame, DisplayRate, TickResolution).GetFrame();
+		SettingsInTick.EndFrame = ConvertFrameTime(SettingsInTick.EndFrame, DisplayRate, TickResolution).GetFrame();
+	}
+
+	TSortedMap<FFrameNumber, FFrameNumber> TotalFrameMap;
 	TMap<FGuid, FBakeData> BakeDataMap;
 	for (const FMovieSceneBindingProxy& ObjectBinding : ObjectBindings)
 	{
 		BakeDataMap.Add(ObjectBinding.BindingID);
 	}
+	CalculateFramesPerGuid(Sequencer, SettingsInTick, BakeDataMap, TotalFrameMap);
 
 	FMovieSceneSequenceTransform RootToLocalTransform = Sequencer->GetFocusedMovieSceneSequenceTransform();
 	
-	for (FFrameTime EvalTime = InFrame; EvalTime <= OutFrame; EvalTime += Interval)
+	TArray<FFrameNumber> AllFrames;
+	TotalFrameMap.GenerateKeyArray(AllFrames);
+
+	for (FFrameNumber KeyTime: AllFrames)
 	{
-		FFrameNumber KeyTime = FFrameRate::Snap(EvalTime, TickResolution, DisplayRate).FloorToFrame();
 		FMovieSceneEvaluationRange Range(KeyTime * RootToLocalTransform.InverseLinearOnly(), TickResolution);
 
 		Sequencer->SetGlobalTime(Range.GetTime());
@@ -748,6 +1048,11 @@ void ULevelSequenceEditorSubsystem::BakeTransform(const TArray<FMovieSceneBindin
 
 			for (TWeakObjectPtr<> RuntimeObject : Sequencer->FindObjectsInCurrentSequence(Guid) )
 			{
+				const FFrameNumber* Number = BakeDataMap[Guid].KeyTimes.FindKey(KeyTime);
+				if (Number == nullptr)
+				{
+					continue;
+				}
 				AActor* Actor = Cast<AActor>(RuntimeObject.Get());
 				if (!Actor)
 				{
@@ -816,7 +1121,6 @@ void ULevelSequenceEditorSubsystem::BakeTransform(const TArray<FMovieSceneBindin
 					BakeDataMap[Guid].Scales.Add(Actor->GetActorScale());
 				}
 
-				BakeDataMap[Guid].KeyTimes.Add(KeyTime);
 			}
 		}
 	}
@@ -826,7 +1130,8 @@ void ULevelSequenceEditorSubsystem::BakeTransform(const TArray<FMovieSceneBindin
 	for (TPair<FGuid, FBakeData>& BakeData : BakeDataMap)
 	{
 		FGuid Guid = BakeData.Key;
-
+		TArray<FFrameNumber> KeyTimes;
+		BakeData.Value.KeyTimes.GenerateKeyArray(KeyTimes);
 		// Disable or delete any constraint (attach/path) tracks
 		AActor* AttachParentActor = nullptr;
 		for (UMovieSceneTrack* Track : FocusedMovieScene->FindTracks(UMovieScene3DConstraintTrack::StaticClass(), Guid))
@@ -953,11 +1258,11 @@ void ULevelSequenceEditorSubsystem::BakeTransform(const TArray<FMovieSceneBindin
 			DoubleChannels[8]->SetDefault(DefaultScale.Z);
 
 			TArray<FVector> LocalTranslations, LocalRotations, LocalScales;
-			LocalTranslations.SetNum(BakeData.Value.KeyTimes.Num());
-			LocalRotations.SetNum(BakeData.Value.KeyTimes.Num());
-			LocalScales.SetNum(BakeData.Value.KeyTimes.Num());
+			LocalTranslations.SetNum(KeyTimes.Num());
+			LocalRotations.SetNum(KeyTimes.Num());
+			LocalScales.SetNum(KeyTimes.Num());
 
-			for (int32 Counter = 0; Counter < BakeData.Value.KeyTimes.Num(); ++Counter)
+			for (int32 Counter = 0; Counter < KeyTimes.Num(); ++Counter)
 			{
 				FTransform LocalTransform(BakeData.Value.Rotations[Counter], BakeData.Value.Locations[Counter], BakeData.Value.Scales[Counter]);
 				LocalTranslations[Counter] = LocalTransform.GetTranslation();
@@ -972,24 +1277,83 @@ void ULevelSequenceEditorSubsystem::BakeTransform(const TArray<FMovieSceneBindin
 				FMath::WindRelativeAnglesDegrees(LocalRotations[Counter].Y, LocalRotations[Counter + 1].Y);
 				FMath::WindRelativeAnglesDegrees(LocalRotations[Counter].Z, LocalRotations[Counter + 1].Z);							
 			}
-				
-			for (int32 Counter = 0; Counter < BakeData.Value.KeyTimes.Num(); ++Counter)
+			if (SettingsInTick.BakingKeySettings == EBakingKeySettings::KeysOnly)
+			{
+				const EMovieSceneKeyInterpolation KeyInterpolation = Sequencer->GetSequencerSettings()->GetKeyInterpolation();
+
+				for (int32 Counter = 0; Counter < KeyTimes.Num(); ++Counter)
+				{
+					int ChannelIndex = 0;
+					FFrameNumber KeyTime = KeyTimes[Counter];
+					{
+						TMovieSceneChannelData<FMovieSceneDoubleValue> ChannelData = DoubleChannels[ChannelIndex++]->GetData();
+						MovieSceneToolHelpers::SetOrAddKey(ChannelData, KeyTime, LocalTranslations[Counter].X);
+					}
+					{
+						TMovieSceneChannelData<FMovieSceneDoubleValue> ChannelData = DoubleChannels[ChannelIndex++]->GetData();
+						MovieSceneToolHelpers::SetOrAddKey(ChannelData, KeyTime, LocalTranslations[Counter].Y);
+					}
+					{
+						TMovieSceneChannelData<FMovieSceneDoubleValue> ChannelData = DoubleChannels[ChannelIndex++]->GetData();
+						MovieSceneToolHelpers::SetOrAddKey(ChannelData, KeyTime, LocalTranslations[Counter].Z);
+					}
+					{
+						TMovieSceneChannelData<FMovieSceneDoubleValue> ChannelData = DoubleChannels[ChannelIndex++]->GetData();
+						MovieSceneToolHelpers::SetOrAddKey(ChannelData, KeyTime, LocalRotations[Counter].X);
+					}
+					{
+						TMovieSceneChannelData<FMovieSceneDoubleValue> ChannelData = DoubleChannels[ChannelIndex++]->GetData();
+						MovieSceneToolHelpers::SetOrAddKey(ChannelData, KeyTime, LocalRotations[Counter].Y);
+
+					}
+					{
+						TMovieSceneChannelData<FMovieSceneDoubleValue> ChannelData = DoubleChannels[ChannelIndex++]->GetData();
+						MovieSceneToolHelpers::SetOrAddKey(ChannelData, KeyTime, LocalRotations[Counter].Z);
+					}
+					{
+						TMovieSceneChannelData<FMovieSceneDoubleValue> ChannelData = DoubleChannels[ChannelIndex++]->GetData();
+						MovieSceneToolHelpers::SetOrAddKey(ChannelData, KeyTime, LocalScales[Counter].X);
+					}
+					{
+						TMovieSceneChannelData<FMovieSceneDoubleValue> ChannelData = DoubleChannels[ChannelIndex++]->GetData();
+						MovieSceneToolHelpers::SetOrAddKey(ChannelData, KeyTime, LocalScales[Counter].Y);
+					}
+					{
+						TMovieSceneChannelData<FMovieSceneDoubleValue> ChannelData = DoubleChannels[ChannelIndex++]->GetData();
+						MovieSceneToolHelpers::SetOrAddKey(ChannelData, KeyTime, LocalScales[Counter].Z);
+					}
+				}
+			}
+			else
 			{
-				FFrameNumber KeyTime = BakeData.Value.KeyTimes[Counter];
-				DoubleChannels[0]->AddLinearKey(KeyTime, LocalTranslations[Counter].X);
-				DoubleChannels[1]->AddLinearKey(KeyTime, LocalTranslations[Counter].Y);
-				DoubleChannels[2]->AddLinearKey(KeyTime, LocalTranslations[Counter].Z);
-				DoubleChannels[3]->AddLinearKey(KeyTime, LocalRotations[Counter].X);
-				DoubleChannels[4]->AddLinearKey(KeyTime, LocalRotations[Counter].Y);
-				DoubleChannels[5]->AddLinearKey(KeyTime, LocalRotations[Counter].Z);
-				DoubleChannels[6]->AddLinearKey(KeyTime, LocalScales[Counter].X);
-				DoubleChannels[7]->AddLinearKey(KeyTime, LocalScales[Counter].Y);
-				DoubleChannels[8]->AddLinearKey(KeyTime, LocalScales[Counter].Z);
+				for (int32 Counter = 0; Counter < KeyTimes.Num(); ++Counter)
+				{
+					FFrameNumber KeyTime = KeyTimes[Counter];
+					DoubleChannels[0]->AddLinearKey(KeyTime, LocalTranslations[Counter].X);
+					DoubleChannels[1]->AddLinearKey(KeyTime, LocalTranslations[Counter].Y);
+					DoubleChannels[2]->AddLinearKey(KeyTime, LocalTranslations[Counter].Z);
+					DoubleChannels[3]->AddLinearKey(KeyTime, LocalRotations[Counter].X);
+					DoubleChannels[4]->AddLinearKey(KeyTime, LocalRotations[Counter].Y);
+					DoubleChannels[5]->AddLinearKey(KeyTime, LocalRotations[Counter].Z);
+					DoubleChannels[6]->AddLinearKey(KeyTime, LocalScales[Counter].X);
+					DoubleChannels[7]->AddLinearKey(KeyTime, LocalScales[Counter].Y);
+					DoubleChannels[8]->AddLinearKey(KeyTime, LocalScales[Counter].Z);
+				}
+				if (SettingsInTick.bReduceKeys == true)
+				{
+					FKeyDataOptimizationParams Param;
+					Param.bAutoSetInterpolation = true;
+					Param.Tolerance = SettingsInTick.Tolerance;
+					TRange<FFrameNumber> Range(SettingsInTick.StartFrame, SettingsInTick.EndFrame);
+					Param.Range = Range;
+					MovieSceneToolHelpers::OptimizeSection(Param, TransformSection);
+				}
 			}
 		}
 	}
 	
 	Sequencer->NotifyMovieSceneDataChanged( EMovieSceneDataChangeType::MovieSceneStructureItemsChanged );
+	return true;
 }
 
 void ULevelSequenceEditorSubsystem::FixActorReferences()
diff --git a/Engine/Plugins/MovieScene/LevelSequenceEditor/Source/LevelSequenceEditor/Public/LevelSequenceEditorSubsystem.h b/Engine/Plugins/MovieScene/LevelSequenceEditor/Source/LevelSequenceEditor/Public/LevelSequenceEditorSubsystem.h
--- a/Engine/Plugins/MovieScene/LevelSequenceEditor/Source/LevelSequenceEditor/Public/LevelSequenceEditorSubsystem.h
+++ b/Engine/Plugins/MovieScene/LevelSequenceEditor/Source/LevelSequenceEditor/Public/LevelSequenceEditorSubsystem.h
@@ -6,7 +6,7 @@
 #include "EditorSubsystem.h"
 
 #include "SequenceTimeUnit.h"
-
+#include "Containers/SortedMap.h"
 #include "LevelSequenceEditorSubsystem.generated.h"
 
 class FUICommandList;
@@ -17,6 +17,7 @@ struct FMovieScenePasteBindingsParams;
 struct FMovieScenePasteFoldersParams;
 struct FMovieScenePasteSectionsParams;
 struct FMovieScenePasteTracksParams;
+struct FBakingAnimationKeySettings;
 
 DECLARE_LOG_CATEGORY_EXTERN(LogLevelSequenceEditor, Log, All);
 
@@ -152,9 +153,13 @@ public:
 	void SyncSectionsUsingSourceTimecode(const TArray<UMovieSceneSection*>& Sections);
 
 	/** Bake transform */
+	UE_DEPRECATED(5.3, "Use ULevelSequenceEditorSubsystem::BakeTransformWithSettings instead")
 	UFUNCTION(BlueprintCallable, Category = "Level Sequence Editor")
 	void BakeTransform(const TArray<FMovieSceneBindingProxy>& ObjectBindings, const FFrameTime& BakeInTime, const FFrameTime& BakeOutTime, const FFrameTime& BakeInterval, const FMovieSceneScriptingParams& Params = FMovieSceneScriptingParams());
 
+	UFUNCTION(BlueprintCallable, Category = "Level Sequence Editor")
+	bool BakeTransformWithSettings(const TArray<FMovieSceneBindingProxy>& ObjectBindings, const FBakingAnimationKeySettings& InSettings, const FMovieSceneScriptingParams& Params = FMovieSceneScriptingParams());
+
 	/** Attempts to automatically fix up broken actor references in the current scene */
 	UFUNCTION(BlueprintCallable, Category = "Level Sequence Editor")
 	void FixActorReferences();
@@ -183,6 +188,18 @@ public:
 	UFUNCTION(BlueprintCallable, Category = "Level Sequence Editor")
 	void RebindComponent(const TArray<FMovieSceneBindingProxy>& ComponentBindings, const FName& ComponentName);
 
+private:
+	/** Used by Baking transforms*/
+	struct FBakeData
+	{
+		TArray<FVector> Locations;
+		TArray<FRotator> Rotations;
+		TArray<FVector> Scales;
+		TSortedMap<FFrameNumber,FFrameNumber> KeyTimes;
+	};
+	void CalculateFramesPerGuid(TSharedPtr<ISequencer>& Sequencer, const FBakingAnimationKeySettings& InSettings, TMap<FGuid, FBakeData>& OutBakeDataMa,
+		TSortedMap<FFrameNumber, FFrameNumber>&  OutFrameMap);
+
 private:
 
 	TSharedPtr<ISequencer> GetActiveSequencer();
diff --git a/Engine/Source/Editor/MovieSceneTools/Private/MovieSceneToolHelpers.cpp b/Engine/Source/Editor/MovieSceneTools/Private/MovieSceneToolHelpers.cpp
--- a/Engine/Source/Editor/MovieSceneTools/Private/MovieSceneToolHelpers.cpp
+++ b/Engine/Source/Editor/MovieSceneTools/Private/MovieSceneToolHelpers.cpp
@@ -28,6 +28,7 @@
 #include "Tracks/MovieSceneFloatTrack.h"
 #include "Tracks/MovieSceneCameraCutTrack.h"
 #include "Sections/MovieScene3DTransformSection.h"
+#include "Sections/MovieScene3DConstraintSection.h"
 #include "Tracks/MovieScene3DTransformTrack.h"
 #include "Sections/MovieSceneCinematicShotSection.h"
 #include "LevelSequence.h"
@@ -57,6 +58,7 @@
 #include "IMovieScenePlayer.h"
 #include "Tracks/MovieSceneCinematicShotTrack.h"
 #include "Tracks/MovieSceneCameraCutTrack.h"
+#include "Tracks/MovieScene3DConstraintTrack.h"
 #include "Sections/MovieSceneCameraCutSection.h"
 #include "Engine/LevelStreaming.h"
 #include "FbxExporter.h"
@@ -84,6 +86,7 @@
 #include "Framework/MultiBox/MultiBoxBuilder.h"
 #include "EntitySystem/Interrogation/MovieSceneInterrogationLinker.h"
 #include "ConstraintsManager.h"
+#include "Constraints/MovieSceneConstraintChannelHelper.inl"
 
 /* FSkelMeshRecorder
  ***********/
@@ -4629,7 +4632,7 @@ bool MovieSceneToolHelpers::IsValidAsset(UMovieSceneSequence* Sequence, const FA
 	return true;
 }
 
-void MovieSceneToolHelpers::SetOrAddKey(TMovieSceneChannelData<FMovieSceneFloatValue>& ChannelData, FFrameNumber Time, float Value)
+void MovieSceneToolHelpers::SetOrAddKey(TMovieSceneChannelData<FMovieSceneFloatValue>& ChannelData, FFrameNumber Time, float Value, const EMovieSceneKeyInterpolation Interpolation)
 {
 	int32 ExistingIndex = ChannelData.FindKey(Time);
 	if (ExistingIndex != INDEX_NONE)
@@ -4648,6 +4651,47 @@ void MovieSceneToolHelpers::SetOrAddKey(TMovieSceneChannelData<FMovieSceneFloatV
 		NewKey.Tangent.TangentWeightMode = WeightedMode;
 		NewKey.Tangent.ArriveTangentWeight = 0.0f;
 		NewKey.Tangent.LeaveTangentWeight = 0.0f;
+
+		switch (Interpolation)
+		{
+			case EMovieSceneKeyInterpolation::SmartAuto:    
+			{
+				NewKey.InterpMode = ERichCurveInterpMode::RCIM_Cubic;
+				NewKey.TangentMode = ERichCurveTangentMode::RCTM_SmartAuto;
+			}
+			break;
+			case EMovieSceneKeyInterpolation::Auto:  
+			{
+				NewKey.InterpMode = ERichCurveInterpMode::RCIM_Cubic;
+				NewKey.TangentMode = ERichCurveTangentMode::RCTM_Auto;
+			}
+			break;
+			case EMovieSceneKeyInterpolation::User: 
+			{
+				NewKey.InterpMode = ERichCurveInterpMode::RCIM_Cubic;
+				NewKey.TangentMode = ERichCurveTangentMode::RCTM_User;
+			}
+			break;
+			case EMovieSceneKeyInterpolation::Break:  
+			{
+				NewKey.InterpMode = ERichCurveInterpMode::RCIM_Cubic;
+				NewKey.TangentMode = ERichCurveTangentMode::RCTM_Auto;
+			}
+			break;
+
+			case EMovieSceneKeyInterpolation::Linear:
+			{
+				NewKey.InterpMode = ERichCurveInterpMode::RCIM_Linear;
+			}
+			break;
+
+			case EMovieSceneKeyInterpolation::Constant:
+			{
+				NewKey.InterpMode = ERichCurveInterpMode::RCIM_Constant;
+			}
+			break;
+
+		}
 		ChannelData.AddKey(Time, NewKey);
 	}
 }
@@ -4712,7 +4756,7 @@ void MovieSceneToolHelpers::SetOrAddKey(TMovieSceneChannelData<FMovieSceneDouble
 	}
 }
 
-void MovieSceneToolHelpers::SetOrAddKey(TMovieSceneChannelData<FMovieSceneDoubleValue>& ChannelData, FFrameNumber Time, double Value)
+void MovieSceneToolHelpers::SetOrAddKey(TMovieSceneChannelData<FMovieSceneDoubleValue>& ChannelData, FFrameNumber Time, double Value, const EMovieSceneKeyInterpolation Interpolation)
 {
 	int32 ExistingIndex = ChannelData.FindKey(Time);
 	if (ExistingIndex != INDEX_NONE)
@@ -4731,6 +4775,47 @@ void MovieSceneToolHelpers::SetOrAddKey(TMovieSceneChannelData<FMovieSceneDouble
 		NewKey.Tangent.TangentWeightMode = WeightedMode;
 		NewKey.Tangent.ArriveTangentWeight = 0.0f;
 		NewKey.Tangent.LeaveTangentWeight = 0.0f;
+
+		switch (Interpolation)
+		{
+			case EMovieSceneKeyInterpolation::SmartAuto:    
+			{
+				NewKey.InterpMode = ERichCurveInterpMode::RCIM_Cubic;
+				NewKey.TangentMode = ERichCurveTangentMode::RCTM_SmartAuto;
+			}
+			break;
+			case EMovieSceneKeyInterpolation::Auto:  
+			{
+				NewKey.InterpMode = ERichCurveInterpMode::RCIM_Cubic;
+				NewKey.TangentMode = ERichCurveTangentMode::RCTM_Auto;
+			}
+			break;
+			case EMovieSceneKeyInterpolation::User: 
+			{
+				NewKey.InterpMode = ERichCurveInterpMode::RCIM_Cubic;
+				NewKey.TangentMode = ERichCurveTangentMode::RCTM_User;
+			}
+			break;
+			case EMovieSceneKeyInterpolation::Break:  
+			{
+				NewKey.InterpMode = ERichCurveInterpMode::RCIM_Cubic;
+				NewKey.TangentMode = ERichCurveTangentMode::RCTM_Auto;
+			}
+			break;
+
+			case EMovieSceneKeyInterpolation::Linear:
+			{
+				NewKey.InterpMode = ERichCurveInterpMode::RCIM_Linear;
+			}
+			break;
+
+			case EMovieSceneKeyInterpolation::Constant:
+			{
+				NewKey.InterpMode = ERichCurveInterpMode::RCIM_Constant;
+			}
+			break;
+
+		}
 		ChannelData.AddKey(Time, NewKey);
 	}
 }
@@ -4772,6 +4857,133 @@ void MovieSceneToolHelpers::GetActorWorldTransforms(IMovieScenePlayer* Player, U
 	}
 }
 
+void MovieSceneToolHelpers::GetActorParents(const FActorForWorldTransforms& ActorSelection,
+	TArray<FActorForWorldTransforms>& OutParentActors)
+{
+	if (ActorSelection.Actor.IsValid())
+	{
+		//has component so parent is next component
+		if (ActorSelection.Component.IsValid() && ActorSelection.Component->GetAttachParent())
+		{
+			FActorForWorldTransforms OutParentActor;
+			OutParentActor.Actor = ActorSelection.Actor;
+			OutParentActor.Component = ActorSelection.Component->GetAttachParent();
+			OutParentActors.AddUnique(OutParentActor);
+			GetActorParents(OutParentActor, OutParentActors);
+		}
+		else if (AActor* ParentActor = ActorSelection.Actor->GetAttachParentActor())
+		{
+			FActorForWorldTransforms OutParentActor;
+			OutParentActor.Actor = ParentActor;
+			OutParentActors.AddUnique(OutParentActor);
+			GetActorParents(OutParentActor, OutParentActors);
+		}
+	}
+}
+
+void MovieSceneToolHelpers::GetActorParentsWithAttachments(ISequencer* Sequencer, const FActorForWorldTransforms& ActorSelection,
+	TArray<FActorForWorldTransforms>& OutParentActors)
+{
+	if (Sequencer == nullptr)
+	{
+		return;
+	}
+
+	UMovieScene* FocusedMovieScene = Sequencer->GetFocusedMovieSceneSequence()->GetMovieScene();
+	if (!FocusedMovieScene)
+	{
+		return;
+	}
+
+	if (ActorSelection.Actor.IsValid())
+	{
+		FGuid Guid;
+		if (ActorSelection.Component.Get())
+		{
+			Guid = Sequencer->GetHandleToObject(ActorSelection.Component.Get(), false);
+		}
+		if (!Guid.IsValid())
+		{
+			Guid = Sequencer->GetHandleToObject(ActorSelection.Actor.Get(), false);
+		}
+		if (Guid.IsValid())
+		{
+			for (UMovieSceneTrack* Track : FocusedMovieScene->FindTracks(UMovieScene3DConstraintTrack::StaticClass(), Guid))
+			{
+				if (UMovieScene3DConstraintTrack* ConstraintTrack = Cast<UMovieScene3DConstraintTrack>(Track))
+				{
+					for (UMovieSceneSection* ConstraintSection : ConstraintTrack->GetAllSections())
+					{
+						FMovieSceneObjectBindingID ConstraintBindingID = (Cast<UMovieScene3DConstraintSection>(ConstraintSection))->GetConstraintBindingID();
+						for (TWeakObjectPtr<> ParentObject : ConstraintBindingID.ResolveBoundObjects(Sequencer->GetFocusedTemplateID(), *Sequencer))
+						{
+							FActorForWorldTransforms OutParentActor;
+							OutParentActor.Actor = Cast<AActor>(ParentObject.Get());;
+							OutParentActors.AddUnique(OutParentActor);
+							GetActorParents(OutParentActor, OutParentActors);
+							break;
+						}
+					}
+				}
+			}
+		}
+		else
+		{
+			//has component so parent is next component
+			if (ActorSelection.Component.IsValid() && ActorSelection.Component->GetAttachParent())
+			{
+				FActorForWorldTransforms OutParentActor;
+				OutParentActor.Actor = ActorSelection.Actor;
+				OutParentActor.Component = ActorSelection.Component->GetAttachParent();
+				OutParentActors.AddUnique(OutParentActor);
+				GetActorParents(OutParentActor, OutParentActors);
+			}
+			else if (AActor* ParentActor = ActorSelection.Actor->GetAttachParentActor())
+			{
+				FActorForWorldTransforms OutParentActor;
+				OutParentActor.Actor = ParentActor;
+				OutParentActors.AddUnique(OutParentActor);
+				GetActorParents(OutParentActor, OutParentActors);
+			}
+		}
+	}
+}
+
+void MovieSceneToolHelpers::GetActorsAndParentsKeyFrames(ISequencer* Sequencer, const FActorForWorldTransforms& InActor,
+	const FFrameNumber& StartFrame, const FFrameNumber& EndFrame, TSortedMap<FFrameNumber, FFrameNumber>& OutFrameMap)
+{
+	TArray<FFrameNumber> FramesToUse;
+	TArray<FActorForWorldTransforms> ParentActors;
+	ParentActors.Add(InActor);
+	GetActorParentsWithAttachments(Sequencer, InActor, ParentActors);
+
+	for (FActorForWorldTransforms& Actor : ParentActors)
+	{
+		FGuid Guid;
+		if (Actor.Component.Get())
+		{
+			Guid = Sequencer->GetHandleToObject(Actor.Component.Get(), false);
+		}
+		if (!Guid.IsValid())
+		{
+			Guid = Sequencer->GetHandleToObject(Actor.Actor.Get(), false);
+		}
+		if (Guid.IsValid())
+		{
+			if (UMovieScene3DTransformSection* TransformSection = MovieSceneToolHelpers::GetTransformSection(Sequencer, Guid))
+			{
+				TArrayView<FMovieSceneDoubleChannel*> Channels = TransformSection->GetChannelProxy().GetChannels<FMovieSceneDoubleChannel>();
+				TArray<FFrameNumber> TransformFrameTimes = FMovieSceneConstraintChannelHelper::GetTransformTimes(
+					Channels, StartFrame, EndFrame);
+				for (FFrameNumber& FrameNumber : TransformFrameTimes)
+				{
+					OutFrameMap.Add(FrameNumber,FrameNumber);
+				}
+			}
+		}
+	}
+}
+
 void MovieSceneToolHelpers::CalculateFramesBetween(
 	const UMovieScene* MovieScene,
 	FFrameNumber StartFrame,
diff --git a/Engine/Source/Editor/MovieSceneTools/Public/MovieSceneToolHelpers.h b/Engine/Source/Editor/MovieSceneTools/Public/MovieSceneToolHelpers.h
--- a/Engine/Source/Editor/MovieSceneTools/Public/MovieSceneToolHelpers.h
+++ b/Engine/Source/Editor/MovieSceneTools/Public/MovieSceneToolHelpers.h
@@ -15,6 +15,7 @@
 #include "MovieSceneTranslator.h"
 #include "MovieSceneSpawnable.h"
 #include "MovieSceneCaptureSettings.h"
+#include "KeyParams.h"
 #include "SEnumCombo.h"
 #include "Animation/AnimSequence.h"
 #include "INodeAndChannelMappings.h"
@@ -533,10 +534,37 @@ public:
 	*/
 	static USkeletalMeshComponent* AcquireSkeletalMeshFromObject(UObject* BoundObject);
 	
+	/*
+	* Get an actors and possible component parents.
+	* @param InActorAndComponent Actor and possible component to find parents for
+	* @param OutParentActors Returns an array of parents
+	*/
+	static void GetActorParents(const FActorForWorldTransforms& Actor,
+		TArray<FActorForWorldTransforms>& OutParentActors);
+
+	/*
+	* Get an actors and possible component parents using sequencer to test for attachments.
+	* @param Sequencer Sequencer to evaluate
+	* @param InActorAndComponent Actor and possible component to find parents for
+	* @param OutParentActors Returns an array of parents
+	*/
+	static void GetActorParentsWithAttachments(ISequencer* Sequencer, const FActorForWorldTransforms& Actor, TArray<FActorForWorldTransforms>& OutParentActors);
+
+	/*
+	*  Get an actors and it's parent key frames
+	* @param Sequencer Sequencer to evaluate
+	* @param Actor The actor and possible component and socket that we want to get the frame for
+	* @param StartFrame The first frame to start looking for keys
+	* @param EndFrame The last frame to stop looking for keys
+	* @param OutFrameMap Sorted map of the frame times found
+	*/
+	static void GetActorsAndParentsKeyFrames(ISequencer* Sequencer, const FActorForWorldTransforms& Actor,
+		const FFrameNumber& StartFrame, const FFrameNumber& EndFrame, TSortedMap<FFrameNumber, FFrameNumber>& OutFrameMap);
+
 	/*
 	*  Get an actors word transforms at the specified times
 	* @param Sequencer Sequencer to evaluate
-    * @param ActorForWorldTransforms The actor and possible component and socket that we want to get the world transforms for.
+    * @param Actors The actor and possible component and socket that we want to get the world transforms for.
 	* @param Frames The times we want to get the world transforms
 	* @param OutWorldTransforms The calculated world transforms, one for each specified frame.
 	*/
@@ -546,26 +574,31 @@ public:
 	* @param ChannelData Channel to set or add
 	* @param Time Frame to add or set the value
 	* @param Value  Value to Set
+	* @param Interpolation Key type to set if added
 	*/
-	static void SetOrAddKey(TMovieSceneChannelData<FMovieSceneFloatValue>& ChannelData, FFrameNumber Time, float Value);
+	static void SetOrAddKey(TMovieSceneChannelData<FMovieSceneFloatValue>& ChannelData, FFrameNumber Time, float Value, const EMovieSceneKeyInterpolation Interpolation = EMovieSceneKeyInterpolation::Auto);
 
 	/* Set or add a key onto a double channel.
 	* @param ChannelData Channel to set or add
 	* @param Time Frame to add or set the value
 	* @param Value  Value to Set
+	* @param Interpolation Key type to set if added
 	*/
-	static void SetOrAddKey(TMovieSceneChannelData<FMovieSceneDoubleValue>& ChannelData, FFrameNumber Time, double Value);
+	static void SetOrAddKey(TMovieSceneChannelData<FMovieSceneDoubleValue>& ChannelData, FFrameNumber Time, double Value, const EMovieSceneKeyInterpolation Interpolation = EMovieSceneKeyInterpolation::Auto);
+
 
 	/*
 	* Set or add a key onto a float channel based on key value.
 	*/
 	static void SetOrAddKey(TMovieSceneChannelData<FMovieSceneFloatValue>& Curve, FFrameNumber Time, const FMovieSceneFloatValue& Value);
 
+
 	/*
 	* Set or add a key onto a double channel based on key value.
 	*/
 	static void SetOrAddKey(TMovieSceneChannelData<FMovieSceneDoubleValue>& ChannelData, FFrameNumber Time, FMovieSceneDoubleValue Value);
 
+
 	/* 
 	* Set or add a key onto a float channel based on rich curve data.
 	*/
diff --git a/Engine/Source/Runtime/MovieScene/Public/ActorForWorldTransforms.h b/Engine/Source/Runtime/MovieScene/Public/ActorForWorldTransforms.h
--- a/Engine/Source/Runtime/MovieScene/Public/ActorForWorldTransforms.h
+++ b/Engine/Source/Runtime/MovieScene/Public/ActorForWorldTransforms.h
@@ -22,4 +22,26 @@ struct  MOVIESCENE_API FActorForWorldTransforms
 
 	UPROPERTY(BlueprintReadWrite, EditAnywhere, Category = "Actor")
 	FName SocketName;
+
+	bool operator ==(const FActorForWorldTransforms& Other) const
+	{
+		if (Actor == Other.Actor)
+		{
+			if (Component.IsValid())
+			{
+				if (Component == Other.Component)
+				{
+					if (SocketName == Other.SocketName)
+					{
+						return true;
+					}
+				}
+			}
+			else if(Other.Component.IsValid() == false)
+			{
+				return true;
+			}
+		}
+		return false;
+	}
 };
\ No newline at end of file
