CIS Unity fix
#trivial

diff --git a/Engine/Plugins/MovieScene/LevelSequenceEditor/Source/LevelSequenceEditor/Private/LevelSequenceEditorSubsystem.cpp b/Engine/Plugins/MovieScene/LevelSequenceEditor/Source/LevelSequenceEditor/Private/LevelSequenceEditorSubsystem.cpp
--- a/Engine/Plugins/MovieScene/LevelSequenceEditor/Source/LevelSequenceEditor/Private/LevelSequenceEditorSubsystem.cpp
+++ b/Engine/Plugins/MovieScene/LevelSequenceEditor/Source/LevelSequenceEditor/Private/LevelSequenceEditorSubsystem.cpp
@@ -23,6 +23,13 @@
 #include "Tracks/MovieSceneCameraShakeTrack.h"
 
 #include "ActorTreeItem.h"
+#include "Editor.h"
+#include "PropertyEditorModule.h"
+#include "Widgets/SWindow.h"
+#include "Widgets/Input/SButton.h"
+#include "Widgets/Layout/SSpacer.h"
+#include "Modules/ModuleManager.h"
+
 #include "Camera/CameraComponent.h"
 #include "Components/SkeletalMeshComponent.h"
 #include "EngineUtils.h"
@@ -45,6 +52,8 @@
 #include "ActorForWorldTransforms.h"
 #include "KeyParams.h"
 
+
+
 #include UE_INLINE_GENERATED_CPP_BY_NAME(LevelSequenceEditorSubsystem)
 
 DEFINE_LOG_CATEGORY(LogLevelSequenceEditor);
