Sequencer: Remove unnecessary creation of camera cut track and section because they're already created when a camera is added

diff --git a/Engine/Plugins/MovieScene/LevelSequenceEditor/Source/LevelSequenceEditor/Private/LevelSequenceEditorToolkit.cpp b/Engine/Plugins/MovieScene/LevelSequenceEditor/Source/LevelSequenceEditor/Private/LevelSequenceEditorToolkit.cpp
--- a/Engine/Plugins/MovieScene/LevelSequenceEditor/Source/LevelSequenceEditor/Private/LevelSequenceEditorToolkit.cpp
+++ b/Engine/Plugins/MovieScene/LevelSequenceEditor/Source/LevelSequenceEditor/Private/LevelSequenceEditorToolkit.cpp
@@ -801,14 +801,6 @@ void FLevelSequenceEditorToolkit::AddShot(UMovieSceneCinematicShotTrack* ShotTra
 		}
 		
 		AddDefaultTracksForActor(*NewCamera, CameraGuid);
-
-		// Create a new camera cut section and add it to the camera cut track
-		CameraCutTrack = ShotSequence->GetMovieScene()->AddCameraCutTrack(UMovieSceneCameraCutTrack::StaticClass());
-		UMovieSceneCameraCutSection* CameraCutSection = NewObject<UMovieSceneCameraCutSection>(CameraCutTrack, NAME_None, RF_Transactional);
-
-		CameraCutSection->SetRange(ShotSequence->GetMovieScene()->GetPlaybackRange());
-		CameraCutSection->SetCameraGuid(CameraGuid);
-		CameraCutTrack->AddSection(*CameraCutSection);
 	}
 }
 
