Added a callback to modify new actor track defaults

diff --git a/Engine/Plugins/MovieScene/LevelSequenceEditor/Source/LevelSequenceEditor/Private/LevelSequenceEditorToolkit.cpp b/Engine/Plugins/MovieScene/LevelSequenceEditor/Source/LevelSequenceEditor/Private/LevelSequenceEditorToolkit.cpp
--- a/Engine/Plugins/MovieScene/LevelSequenceEditor/Source/LevelSequenceEditor/Private/LevelSequenceEditorToolkit.cpp
+++ b/Engine/Plugins/MovieScene/LevelSequenceEditor/Source/LevelSequenceEditor/Private/LevelSequenceEditorToolkit.cpp
@@ -8,6 +8,7 @@
 #include "Components/SkeletalMeshComponent.h"
 #include "Containers/ArrayBuilder.h"
 #include "KeyParams.h"
+#include "ILevelSequenceModule.h"
 #include "LevelSequenceEditorModule.h"
 #include "LevelSequenceEditorSubsystem.h"
 #include "Misc/LevelSequenceEditorSettings.h"
@@ -625,6 +626,10 @@ void FLevelSequenceEditorToolkit::AddDefaultTracksForActor(AActor& Actor, const
 			Sequencer->KeyProperty(KeyPropertyParams);
 		}
 	}
+
+	// callback to set up default tracks via code
+	ILevelSequenceModule& LevelSequenceModule = FModuleManager::LoadModuleChecked<ILevelSequenceModule>("LevelSequence");
+	LevelSequenceModule.OnNewActorTrackAdded().Broadcast(Actor, Binding, Sequencer);
 }
 
 
diff --git a/Engine/Source/Runtime/LevelSequence/Private/LevelSequenceModule.cpp b/Engine/Source/Runtime/LevelSequence/Private/LevelSequenceModule.cpp
--- a/Engine/Source/Runtime/LevelSequence/Private/LevelSequenceModule.cpp
+++ b/Engine/Source/Runtime/LevelSequence/Private/LevelSequenceModule.cpp
@@ -343,6 +343,11 @@ void FLevelSequenceModule::UnregisterObjectSpawner(FDelegateHandle InHandle)
 	OnCreateMovieSceneObjectSpawnerDelegates.RemoveAll([=](const FOnCreateMovieSceneObjectSpawner& Delegate) { return Delegate.GetHandle() == InHandle; });
 }
 
+FLevelSequenceModule::FOnNewActorTrackAdded& FLevelSequenceModule::OnNewActorTrackAdded()
+{
+	return NewActorTrackAdded;
+}
+
 void FLevelSequenceModule::GenerateObjectSpawners(TArray<TSharedRef<IMovieSceneObjectSpawner>>& OutSpawners) const
 {
 	for (const FOnCreateMovieSceneObjectSpawner& SpawnerFactory : OnCreateMovieSceneObjectSpawnerDelegates)
diff --git a/Engine/Source/Runtime/LevelSequence/Private/LevelSequenceModule.h b/Engine/Source/Runtime/LevelSequence/Private/LevelSequenceModule.h
--- a/Engine/Source/Runtime/LevelSequence/Private/LevelSequenceModule.h
+++ b/Engine/Source/Runtime/LevelSequence/Private/LevelSequenceModule.h
@@ -24,6 +24,7 @@ public:
 	virtual FDelegateHandle RegisterObjectSpawner(FOnCreateMovieSceneObjectSpawner InOnCreateMovieSceneObjectSpawner) override;
 	virtual void GenerateObjectSpawners(TArray<TSharedRef<IMovieSceneObjectSpawner>>& OutSpawners) const override;
 	virtual void UnregisterObjectSpawner(FDelegateHandle InHandle) override;
+	virtual FOnNewActorTrackAdded& OnNewActorTrackAdded() override;
 
 protected:
 	// FSelfRegisteringExec interface
@@ -40,4 +41,5 @@ private:
 	// Weak ptr to the level sequence CDO so we can gracefully remove the meta-data on shutdown module
 	// without crashing when ShutdownModule is called after the CDO has been destroyed.
 	TWeakObjectPtr<ULevelSequence> LevelSequenceCDO;
+	FOnNewActorTrackAdded NewActorTrackAdded;
 };
diff --git a/Engine/Source/Runtime/LevelSequence/Public/ILevelSequenceModule.h b/Engine/Source/Runtime/LevelSequence/Public/ILevelSequenceModule.h
--- a/Engine/Source/Runtime/LevelSequence/Public/ILevelSequenceModule.h
+++ b/Engine/Source/Runtime/LevelSequence/Public/ILevelSequenceModule.h
@@ -21,4 +21,8 @@ public:
 
 	/** Unregister an object spawner */
 	virtual void UnregisterObjectSpawner(FDelegateHandle InHandle) = 0;
+
+	DECLARE_MULTICAST_DELEGATE_ThreeParams(FOnNewActorTrackAdded, const AActor& /*SourceActor*/, const FGuid& /*Binding*/, TSharedPtr<ISequencer> /*Sequencer*/);
+	/** Callback to set up defaults for new actor tracks */
+	virtual FOnNewActorTrackAdded& OnNewActorTrackAdded() = 0;
 };
