Sequencer: move level-sequencer-specific editor code to the LevelSequenceEditor plugin

diff --git a/Engine/Plugins/MovieScene/LevelSequenceEditor/Source/LevelSequenceEditor/LevelSequenceEditor.Build.cs b/Engine/Plugins/MovieScene/LevelSequenceEditor/Source/LevelSequenceEditor/LevelSequenceEditor.Build.cs
--- a/Engine/Plugins/MovieScene/LevelSequenceEditor/Source/LevelSequenceEditor/LevelSequenceEditor.Build.cs
+++ b/Engine/Plugins/MovieScene/LevelSequenceEditor/Source/LevelSequenceEditor/LevelSequenceEditor.Build.cs
@@ -6,6 +6,12 @@ public class LevelSequenceEditor : ModuleRules
 {
 	public LevelSequenceEditor(ReadOnlyTargetRules Target) : base(Target)
 	{
+		PrivateIncludePaths.AddRange(
+			new string[] {
+                "../../../../Source/Editor/UnrealEd/Private", // TODO: Fix this, for now it's needed for the fbx exporter
+				}
+			);
+
         DynamicallyLoadedModuleNames.AddRange(
             new string[] {
 				"AssetTools",
@@ -61,6 +67,7 @@ public class LevelSequenceEditor : ModuleRules
 				"PlacementMode",
                 "Settings",
                 "MovieSceneCaptureDialog",
+				"DesktopPlatform",
 			}
 		);
 
@@ -72,5 +79,7 @@ public class LevelSequenceEditor : ModuleRules
 				"LevelSequenceEditor/Private/Styles",
 			}
         );
+
+		AddEngineThirdPartyPrivateStaticDependencies(Target, "FBX");
 	}
 }
diff --git a/Engine/Plugins/MovieScene/LevelSequenceEditor/Source/LevelSequenceEditor/Private/LevelSequenceEditorCommands.cpp b/Engine/Plugins/MovieScene/LevelSequenceEditor/Source/LevelSequenceEditor/Private/LevelSequenceEditorCommands.cpp
--- a/Engine/Plugins/MovieScene/LevelSequenceEditor/Source/LevelSequenceEditor/Private/LevelSequenceEditorCommands.cpp
+++ b/Engine/Plugins/MovieScene/LevelSequenceEditor/Source/LevelSequenceEditor/Private/LevelSequenceEditorCommands.cpp
@@ -26,6 +26,9 @@ void FLevelSequenceEditorCommands::RegisterCommands()
 	UI_COMMAND(ReplaceBindingWithActors, "Replace with Selected", "Replace the object binding with selected actors", EUserInterfaceActionType::Button, FInputChord());
 	UI_COMMAND(RemoveAllBindings, "Remove All", "Remove all bound actors from this track", EUserInterfaceActionType::Button, FInputChord());
 	UI_COMMAND(RemoveInvalidBindings, "Remove Missing", "Remove missing objects bound to this track", EUserInterfaceActionType::Button, FInputChord());
+
+	UI_COMMAND(ImportFBX, "Import...", "Import the animation from an FBX file.", EUserInterfaceActionType::Button, FInputChord());
+	UI_COMMAND(ExportFBX, "Export...", "Export the selected objects (all if none selected) and animation to an FBX file. (Shots and sub-scenes not supported)", EUserInterfaceActionType::Button, FInputChord());
 }
 
 #undef LOCTEXT_NAMESPACE
diff --git a/Engine/Plugins/MovieScene/LevelSequenceEditor/Source/LevelSequenceEditor/Private/LevelSequenceEditorCommands.h b/Engine/Plugins/MovieScene/LevelSequenceEditor/Source/LevelSequenceEditor/Private/LevelSequenceEditorCommands.h
--- a/Engine/Plugins/MovieScene/LevelSequenceEditor/Source/LevelSequenceEditor/Private/LevelSequenceEditorCommands.h
+++ b/Engine/Plugins/MovieScene/LevelSequenceEditor/Source/LevelSequenceEditor/Private/LevelSequenceEditorCommands.h
@@ -30,4 +30,11 @@ public:
 	TSharedPtr<FUICommandInfo> ReplaceBindingWithActors;
 	TSharedPtr<FUICommandInfo> RemoveAllBindings;
 	TSharedPtr<FUICommandInfo> RemoveInvalidBindings;
+
+	/** Imports animation from fbx. */
+	TSharedPtr< FUICommandInfo > ImportFBX;
+
+	/** Exports animation to fbx. */
+	TSharedPtr< FUICommandInfo > ExportFBX;
+
 };
diff --git a/Engine/Plugins/MovieScene/LevelSequenceEditor/Source/LevelSequenceEditor/Private/LevelSequenceEditorModule.cpp b/Engine/Plugins/MovieScene/LevelSequenceEditor/Source/LevelSequenceEditor/Private/LevelSequenceEditorModule.cpp
--- a/Engine/Plugins/MovieScene/LevelSequenceEditor/Source/LevelSequenceEditor/Private/LevelSequenceEditorModule.cpp
+++ b/Engine/Plugins/MovieScene/LevelSequenceEditor/Source/LevelSequenceEditor/Private/LevelSequenceEditorModule.cpp
@@ -1,33 +1,37 @@
 // Copyright Epic Games, Inc. All Rights Reserved.
 
 #include "LevelSequenceEditorModule.h"
-#include "Blueprint/BlueprintSupport.h"
-#include "Factories/Factory.h"
-#include "AssetToolsModule.h"
-#include "Framework/MultiBox/MultiBoxBuilder.h"
 
-#include "IAssetTools.h"
-#include "LevelSequenceEditorCommands.h"
-#include "Misc/LevelSequenceEditorSettings.h"
-#include "Misc/LevelSequenceEditorHelpers.h"
-#include "Styles/LevelSequenceEditorStyle.h"
+#include "AssetToolsModule.h"
+#include "Blueprint/BlueprintSupport.h"
+#include "BlueprintAssetHandler.h"
 #include "Camera/CameraShakeSourceActor.h"
-#include "CineCameraActor.h"
 #include "CameraRig_Crane.h"
 #include "CameraRig_Rail.h"
+#include "CineCameraActor.h"
+#include "CinematicViewport/CinematicViewportLayoutEntity.h"
+#include "Factories/Factory.h"
+#include "Framework/MultiBox/MultiBoxBuilder.h"
+#include "Framework/MultiBox/MultiBoxExtender.h"
+#include "IAssetTools.h"
+#include "ILevelSequenceModule.h"
 #include "IPlacementModeModule.h"
+#include "ISequencerModule.h"
 #include "ISettingsModule.h"
 #include "LevelEditor.h"
 #include "LevelSequenceActor.h"
-#include "UObject/UObjectIterator.h"
-#include "CinematicViewport/CinematicViewportLayoutEntity.h"
-#include "ISequencerModule.h"
+#include "LevelSequenceEditorCommands.h"
+#include "Misc/LevelSequenceCustomization.h"
 #include "Misc/LevelSequenceEditorActorBinding.h"
-#include "ILevelSequenceModule.h"
 #include "Misc/LevelSequenceEditorActorSpawner.h"
-#include "SequencerSettings.h"
+#include "Misc/LevelSequenceEditorHelpers.h"
+#include "Misc/LevelSequenceEditorSettings.h"
+#include "Misc/LevelSequenceFBXInterop.h"
 #include "Misc/MovieSceneSequenceEditor_LevelSequence.h"
-#include "BlueprintAssetHandler.h"
+#include "SequencerCommands.h"
+#include "SequencerSettings.h"
+#include "Styles/LevelSequenceEditorStyle.h"
+#include "UObject/UObjectIterator.h"
 
 #define LOCTEXT_NAMESPACE "LevelSequenceEditor"
 
@@ -61,6 +65,7 @@ public:
 		RegisterPlacementModeExtensions();
 		RegisterSettings();
 		RegisterSequenceEditor();
+		RegisterSequenceCustomizations();
 
 		class FLevelSequenceAssetBlueprintHandler : public IBlueprintAssetHandler
 		{
@@ -88,6 +93,7 @@ public:
 		UnregisterPlacementModeExtensions();
 		UnregisterSettings();
 		UnregisterSequenceEditor();
+		UnregisterSequenceCustomizations();
 	}
 
 protected:
@@ -118,13 +124,16 @@ protected:
 	/** Register menu extensions for the level editor toolbar. */
 	void RegisterMenuExtensions()
 	{
+		FSequencerCommands::Register();
 		FLevelSequenceEditorCommands::Register();
 
+		const FLevelSequenceEditorCommands& Commands = FLevelSequenceEditorCommands::Get();
+
 		CommandList = MakeShareable(new FUICommandList);
-		CommandList->MapAction(FLevelSequenceEditorCommands::Get().CreateNewLevelSequenceInLevel,
+		CommandList->MapAction(Commands.CreateNewLevelSequenceInLevel,
 			FExecuteAction::CreateStatic(&FLevelSequenceEditorModule::OnCreateActorInLevel)
 		);
-		CommandList->MapAction(FLevelSequenceEditorCommands::Get().CreateNewLevelSequenceWithShotsInLevel,
+		CommandList->MapAction(Commands.CreateNewLevelSequenceWithShotsInLevel,
 			FExecuteAction::CreateStatic(&FLevelSequenceEditorModule::OnCreateLevelSequenceWithShotsInLevel)
 		);
 
@@ -188,6 +197,16 @@ protected:
 		SequenceEditorHandle = SequencerModule.RegisterSequenceEditor(ULevelSequence::StaticClass(), MakeUnique<FMovieSceneSequenceEditor_LevelSequence>());
 	}
 
+	void RegisterSequenceCustomizations()
+	{
+		ISequencerModule& SequencerModule = FModuleManager::LoadModuleChecked<ISequencerModule>("Sequencer");
+		SequencerModule.GetSequencerCustomizationManager()->RegisterInstancedSequencerCustomization(ULevelSequence::StaticClass(),
+				FOnGetSequencerCustomizationInstance::CreateLambda([]()
+				{
+					return new UE::Sequencer::FLevelSequenceCustomization();
+				}));
+	}
+
 protected:
 
 	/** Unregisters sequencer editor object bindings */
@@ -226,7 +245,7 @@ protected:
 		if (FLevelEditorModule* LevelEditorModule = FModuleManager::GetModulePtr<FLevelEditorModule>("LevelEditor"))
 		{
 			LevelEditorModule->GetAllLevelEditorToolbarCinematicsMenuExtenders().Remove(CinematicsMenuExtender);
-	}
+		}
 
 		CinematicsMenuExtender = nullptr;
 		CommandList = nullptr;
@@ -265,6 +284,12 @@ protected:
 		}
 	}
 
+	void UnregisterSequenceCustomizations()
+	{
+		ISequencerModule& SequencerModule = FModuleManager::LoadModuleChecked<ISequencerModule>("Sequencer");
+		SequencerModule.GetSequencerCustomizationManager()->UnregisterInstancedSequencerCustomization(ULevelSequence::StaticClass());
+	}
+
 protected:
 
 	/** Callback for creating a new level sequence asset in the level. */
diff --git a/Engine/Plugins/MovieScene/LevelSequenceEditor/Source/LevelSequenceEditor/Private/Misc/LevelSequenceCustomization.cpp b/Engine/Plugins/MovieScene/LevelSequenceEditor/Source/LevelSequenceEditor/Private/Misc/LevelSequenceCustomization.cpp
new file mode 100644
--- /dev/null
+++ b/Engine/Plugins/MovieScene/LevelSequenceEditor/Source/LevelSequenceEditor/Private/Misc/LevelSequenceCustomization.cpp
@@ -0,0 +1,548 @@
+// Copyright Epic Games, Inc. All Rights Reserved.
+
+#include "Misc/LevelSequenceCustomization.h"
+
+#include "ClassViewerModule.h"
+#include "Engine/LevelStreaming.h"
+#include "Framework/Application/SlateApplication.h"
+#include "Framework/MultiBox/MultiBoxBuilder.h"
+#include "Framework/MultiBox/MultiBoxExtender.h"
+#include "ISequencer.h"
+#include "ISequencerModule.h"
+#include "LevelSequenceEditorCommands.h"
+#include "Misc/LevelSequenceFBXInterop.h"
+#include "Modules/ModuleManager.h"
+#include "MovieScene.h"
+#include "MVVM/Extensions/IObjectBindingExtension.h"
+#include "MVVM/Selection/Selection.h"
+#include "MVVM/ViewModels/ObjectBindingModel.h"
+#include "MVVM/ViewModels/SequencerEditorViewModel.h"
+#include "ScopedTransaction.h"
+#include "SequencerCommands.h"
+#include "SequencerUtilities.h"
+#include "Tracks/MovieSceneSpawnTrack.h"
+
+#define LOCTEXT_NAMESPACE "LevelSequenceCustomization"
+
+namespace UE::Sequencer
+{
+
+struct FMovieSceneSpawnableFlagCheckState
+{
+	TSharedPtr<FSequencerEditorViewModel> EditorViewModel;
+	UMovieScene* MovieScene;
+	bool FMovieSceneSpawnable::* PtrToFlag;
+
+	ECheckBoxState operator()() const
+	{
+		ECheckBoxState CheckState = ECheckBoxState::Undetermined;
+		for (TViewModelPtr<IObjectBindingExtension> ObjectBinding : EditorViewModel->GetSelection()->Outliner.Filter<IObjectBindingExtension>())
+		{
+			FMovieSceneSpawnable* SelectedSpawnable = MovieScene->FindSpawnable(ObjectBinding->GetObjectGuid());
+			if (SelectedSpawnable)
+			{
+				if (CheckState != ECheckBoxState::Undetermined && SelectedSpawnable->*PtrToFlag != (CheckState == ECheckBoxState::Checked))
+				{
+					return ECheckBoxState::Undetermined;
+				}
+				CheckState = SelectedSpawnable->*PtrToFlag ? ECheckBoxState::Checked : ECheckBoxState::Unchecked;
+			}
+		}
+		return CheckState;
+	}
+};
+
+struct FMovieSceneSpawnableFlagToggler
+{
+	TSharedPtr<FSequencerEditorViewModel> EditorViewModel;
+	UMovieScene* MovieScene;
+	bool FMovieSceneSpawnable::* PtrToFlag;
+	FText TransactionText;
+
+	void operator()() const
+	{
+		FScopedTransaction Transaction(TransactionText);
+
+		const ECheckBoxState CheckState = FMovieSceneSpawnableFlagCheckState{ EditorViewModel, MovieScene, PtrToFlag }();
+
+		MovieScene->Modify();
+		for (TViewModelPtr<IObjectBindingExtension> ObjectBinding : EditorViewModel->GetSelection()->Outliner.Filter<IObjectBindingExtension>())
+		{
+			FMovieSceneSpawnable* SelectedSpawnable = MovieScene->FindSpawnable(ObjectBinding->GetObjectGuid());
+			if (SelectedSpawnable)
+			{
+				SelectedSpawnable->*PtrToFlag = (CheckState == ECheckBoxState::Unchecked);
+			}
+		}
+	}
+};
+
+void FLevelSequenceCustomization::RegisterSequencerCustomization(FSequencerCustomizationBuilder& Builder)
+{
+	WeakSequencer = Builder.GetSequencer().AsShared();
+
+	const FLevelSequenceEditorCommands& Commands = FLevelSequenceEditorCommands::Get();
+
+	// Build the extender for the actions menu.
+	ActionsMenuCommandList = MakeShared<FUICommandList>().ToSharedPtr();
+	ActionsMenuCommandList->MapAction(
+		Commands.ImportFBX,
+		FExecuteAction::CreateRaw( this, &FLevelSequenceCustomization::ImportFBX ),
+		FCanExecuteAction::CreateLambda( [] { return true; } ) );
+	ActionsMenuCommandList->MapAction(
+		Commands.ExportFBX,
+		FExecuteAction::CreateRaw( this, &FLevelSequenceCustomization::ExportFBX ),
+		FCanExecuteAction::CreateLambda( [] { return true; } ) );
+
+	ActionsMenuExtender = MakeShared<FExtender>();
+	ActionsMenuExtender->AddMenuExtension(
+			"SequenceOptions", EExtensionHook::First, nullptr,
+			FMenuExtensionDelegate::CreateRaw(this, &FLevelSequenceCustomization::ExtendActionsMenu));
+
+	ISequencerModule& SequencerModule = FModuleManager::LoadModuleChecked<ISequencerModule>("Sequencer");
+	SequencerModule.GetActionsMenuExtensibilityManager()->AddExtender(ActionsMenuExtender);
+
+	// Add a customization callback for the object binding context menu.
+	FSequencerCustomizationInfo Customization;
+	Customization.OnBuildObjectBindingContextMenu = FOnGetSequencerMenuExtender::CreateRaw(this, &FLevelSequenceCustomization::CreateObjectBindingContextMenuExtender);
+	Builder.AddCustomization(Customization);
+}
+
+void FLevelSequenceCustomization::UnregisterSequencerCustomization()
+{
+	ISequencerModule& SequencerModule = FModuleManager::LoadModuleChecked<ISequencerModule>("Sequencer");
+	SequencerModule.GetActionsMenuExtensibilityManager()->RemoveExtender(ActionsMenuExtender);
+
+	ActionsMenuCommandList = nullptr;
+	WeakSequencer = nullptr;
+}
+
+void FLevelSequenceCustomization::ExtendActionsMenu(FMenuBuilder& MenuBuilder)
+{
+	MenuBuilder.PushCommandList(ActionsMenuCommandList.ToSharedRef());
+	{
+		const FLevelSequenceEditorCommands& Commands = FLevelSequenceEditorCommands::Get();
+		
+		MenuBuilder.AddMenuEntry(
+				LOCTEXT("SaveAs", "Save As..."),
+				LOCTEXT("SaveAsTooltip", "Saves the current sequence under a different name"),
+				FSlateIcon(Commands.GetStyleSetName(), "LevelSequenceEditor.SaveAs"),
+				FUIAction(FExecuteAction::CreateRaw(this, &FLevelSequenceCustomization::OnSaveMovieSceneAsClicked)));
+
+		MenuBuilder.AddMenuEntry(Commands.ImportFBX);
+		MenuBuilder.AddMenuEntry(Commands.ExportFBX);
+	}
+	MenuBuilder.PopCommandList();
+}
+
+void FLevelSequenceCustomization::OnSaveMovieSceneAsClicked()
+{
+	FSequencerUtilities::SaveCurrentMovieSceneAs(WeakSequencer.Pin().ToSharedRef());
+}
+
+void FLevelSequenceCustomization::ImportFBX()
+{
+	FLevelSequenceFBXInterop Interop(WeakSequencer.Pin());
+	Interop.ImportFBXOntoSelectedNodes();
+}
+
+void FLevelSequenceCustomization::ExportFBX()
+{
+	FLevelSequenceFBXInterop Interop(WeakSequencer.Pin());
+	Interop.ExportFBX();
+}
+
+TSharedPtr<FExtender> FLevelSequenceCustomization::CreateObjectBindingContextMenuExtender(FViewModelPtr InViewModel)
+{
+	TSharedRef<FExtender> Extender = MakeShared<FExtender>();
+	TSharedPtr<FObjectBindingModel> ObjectBindingModel = InViewModel->CastThisShared<FObjectBindingModel>();
+	Extender->AddMenuExtension(
+			"ObjectBindingActions", EExtensionHook::Before, nullptr,
+			FMenuExtensionDelegate::CreateRaw(this, &FLevelSequenceCustomization::ExtendObjectBindingContextMenu, ObjectBindingModel));
+	return Extender.ToSharedPtr();
+}
+
+void FLevelSequenceCustomization::ExtendObjectBindingContextMenu(FMenuBuilder& MenuBuilder, TSharedPtr<FObjectBindingModel> ObjectBindingModel)
+{
+	TSharedPtr<ISequencer> Sequencer = WeakSequencer.Pin();
+	TSharedPtr<FSequencerEditorViewModel> EditorViewModel = Sequencer->GetViewModel();
+
+	FGuid ObjectBindingID = ObjectBindingModel->GetObjectGuid();
+	UMovieScene* MovieScene = Sequencer->GetFocusedMovieSceneSequence()->GetMovieScene();
+
+	if (!MovieScene || !ObjectBindingID.IsValid())
+	{
+		return;
+	}
+
+	FMovieSceneSpawnable* Spawnable = MovieScene->FindSpawnable(ObjectBindingID);
+
+	if (Spawnable)
+	{
+		MenuBuilder.BeginSection("Spawnable", LOCTEXT("SpawnableMenuSectionName", "Spawnable"));
+
+		MenuBuilder.AddSubMenu(
+			LOCTEXT("OwnerLabel", "Spawned Object Owner"),
+			LOCTEXT("OwnerTooltip", "Specifies how the spawned object is to be owned"),
+			FNewMenuDelegate::CreateRaw(this, &FLevelSequenceCustomization::AddSpawnOwnershipMenu, ObjectBindingModel)
+		);
+
+		MenuBuilder.AddSubMenu(
+			LOCTEXT("SubLevelLabel", "Spawnable Level"),
+			LOCTEXT("SubLevelTooltip", "Specifies which level the spawnable should be spawned into"),
+			FNewMenuDelegate::CreateRaw(this, &FLevelSequenceCustomization::AddSpawnLevelMenu, ObjectBindingModel)
+		);
+
+		MenuBuilder.AddSubMenu(
+			LOCTEXT("ChangeClassLabel", "Change Class"),
+			LOCTEXT("ChangeClassTooltip", "Change the class (object template) that this spawns from"),
+			FNewMenuDelegate::CreateRaw(this, &FLevelSequenceCustomization::AddChangeClassMenu, ObjectBindingModel));
+
+		MenuBuilder.AddMenuEntry(
+			LOCTEXT("ContinuouslyRespawn", "Continuously Respawn"),
+			LOCTEXT("ContinuouslyRespawnTooltip", "When enabled, this spawnable will always be respawned if it gets destroyed externally. When disabled, this object will only ever be spawned once for each spawn key even if it is destroyed externally"),
+			FSlateIcon(),
+			FUIAction(
+				FExecuteAction::CreateLambda(FMovieSceneSpawnableFlagToggler{ EditorViewModel, MovieScene, &FMovieSceneSpawnable::bContinuouslyRespawn, LOCTEXT("ContinuouslyRespawnTransaction", "Set Continuously Respawn") }),
+				FCanExecuteAction(),
+				FGetActionCheckState::CreateLambda(FMovieSceneSpawnableFlagCheckState{ EditorViewModel, MovieScene, &FMovieSceneSpawnable::bContinuouslyRespawn })
+			),
+			NAME_None,
+			EUserInterfaceActionType::ToggleButton
+		);
+
+		MenuBuilder.AddMenuEntry(
+			LOCTEXT("NetAddressable", "Net Addressable"),
+			LOCTEXT("NetAddressableTooltip", "When enabled, this spawnable will be spawned using a unique name that allows it to be addressed by the server and client (useful for relative movement calculations on spawned props)"),
+			FSlateIcon(),
+			FUIAction(
+				FExecuteAction::CreateLambda(FMovieSceneSpawnableFlagToggler{ EditorViewModel, MovieScene, &FMovieSceneSpawnable::bNetAddressableName, LOCTEXT("NetAddressableTransaction", "Set Net Addressable") }),
+				FCanExecuteAction(),
+				FGetActionCheckState::CreateLambda(FMovieSceneSpawnableFlagCheckState{ EditorViewModel, MovieScene, &FMovieSceneSpawnable::bNetAddressableName })
+			),
+			NAME_None,
+			EUserInterfaceActionType::ToggleButton
+		);
+
+		MenuBuilder.AddMenuEntry(FSequencerCommands::Get().SaveCurrentSpawnableState);
+		MenuBuilder.AddMenuEntry(FSequencerCommands::Get().ConvertToPossessable);
+
+		MenuBuilder.AddSubMenu(
+				LOCTEXT("DynamicSpawn", "Dynamic Spawn"),
+				LOCTEXT("DynamicSpawnTooltip", "Specify a Blueprint method that will spawn or otherwise acquire a compatible actor for this binding"),
+				FNewMenuDelegate::CreateRaw(this, &FLevelSequenceCustomization::AddDynamicSpawnMenu, ObjectBindingModel));
+
+		MenuBuilder.EndSection();
+	}
+	else
+	{
+		MenuBuilder.BeginSection("Possessable");
+
+		MenuBuilder.AddMenuEntry(FSequencerCommands::Get().ConvertToSpawnable);
+
+		MenuBuilder.AddSubMenu(
+				LOCTEXT("DynamicPossession", "Dynamic Possession"),
+				LOCTEXT("DynamicPossessionTooltip", "Specify a Blueprint method that will find a compatible actor for this binding"),
+				FNewMenuDelegate::CreateRaw(this, &FLevelSequenceCustomization::AddDynamicPossessionMenu, ObjectBindingModel));
+
+		MenuBuilder.EndSection();
+	}
+
+	MenuBuilder.BeginSection("Import/Export", LOCTEXT("ImportExportMenuSectionName", "Import/Export"));
+
+	MenuBuilder.AddMenuEntry(
+		LOCTEXT("ImportFBX", "Import..."),
+		LOCTEXT("ImportFBXTooltip", "Import FBX animation to this object"),
+		FSlateIcon(),
+		FUIAction(
+			FExecuteAction::CreateLambda([=] {
+					FLevelSequenceFBXInterop Interop(EditorViewModel->GetSequencer());
+					Interop.ImportFBXOntoSelectedNodes();
+				})
+		));
+
+	MenuBuilder.AddMenuEntry(
+		LOCTEXT("ExportFBX", "Export..."),
+		LOCTEXT("ExportFBXTooltip", "Export FBX animation from this object"),
+		FSlateIcon(),
+		FUIAction(
+			FExecuteAction::CreateLambda([=] {
+					FLevelSequenceFBXInterop Interop(EditorViewModel->GetSequencer());
+					Interop.ExportFBX();
+				})
+		));
+
+	MenuBuilder.EndSection();
+}
+
+void FLevelSequenceCustomization::AddSpawnOwnershipMenu(FMenuBuilder& MenuBuilder, TSharedPtr<FObjectBindingModel> ObjectBindingModel)
+{
+	TSharedPtr<ISequencer> Sequencer = WeakSequencer.Pin();
+	TSharedPtr<FSequencerEditorViewModel> EditorViewModel = Sequencer->GetViewModel();
+
+	FGuid ObjectBindingID = ObjectBindingModel->GetObjectGuid();
+	UMovieScene* MovieScene = Sequencer->GetFocusedMovieSceneSequence()->GetMovieScene();
+
+	FMovieSceneSpawnable* Spawnable = MovieScene->FindSpawnable(ObjectBindingID);
+	if (!Spawnable)
+	{
+		return;
+	}
+
+	auto Callback = [=](ESpawnOwnership NewOwnership)
+	{
+		FScopedTransaction Transaction(LOCTEXT("SetSpawnOwnership", "Set Spawnable Ownership"));
+
+		MovieScene->Modify();
+
+		for (TViewModelPtr<IObjectBindingExtension> ObjectBinding : Sequencer->GetViewModel()->GetSelection()->Outliner.Filter<IObjectBindingExtension>())
+		{
+			FMovieSceneSpawnable* Spawnable = MovieScene->FindSpawnable(ObjectBinding->GetObjectGuid());
+
+			Spawnable->SetSpawnOwnership(NewOwnership);
+
+			// Overwrite the completion state for all spawn sections to ensure the expected behaviour.
+			EMovieSceneCompletionMode NewCompletionMode = NewOwnership == ESpawnOwnership::InnerSequence ? EMovieSceneCompletionMode::RestoreState : EMovieSceneCompletionMode::KeepState;
+
+			// Make all spawn sections retain state
+			UMovieSceneSpawnTrack* SpawnTrack = MovieScene->FindTrack<UMovieSceneSpawnTrack>(ObjectBinding->GetObjectGuid());
+			if (SpawnTrack)
+			{
+				for (UMovieSceneSection* Section : SpawnTrack->GetAllSections())
+				{
+					Section->Modify();
+					Section->EvalOptions.CompletionMode = NewCompletionMode;
+				}
+			}
+		}
+	};
+
+	MenuBuilder.AddMenuEntry(
+		LOCTEXT("ThisSequence_Label", "This Sequence"),
+		LOCTEXT("ThisSequence_Tooltip", "Indicates that this sequence will own the spawned object. The object will be destroyed at the end of the sequence."),
+		FSlateIcon(),
+		FUIAction(
+			FExecuteAction::CreateLambda(Callback, ESpawnOwnership::InnerSequence),
+			FCanExecuteAction(),
+			FIsActionChecked::CreateLambda([=] { return Spawnable->GetSpawnOwnership() == ESpawnOwnership::InnerSequence; })
+		),
+		NAME_None,
+		EUserInterfaceActionType::ToggleButton
+	);
+
+	MenuBuilder.AddMenuEntry(
+		LOCTEXT("RootSequence_Label", "Root Sequence"),
+		LOCTEXT("RootSequence_Tooltip", "Indicates that the outermost sequence will own the spawned object. The object will be destroyed when the outermost sequence stops playing."),
+		FSlateIcon(),
+		FUIAction(
+			FExecuteAction::CreateLambda(Callback, ESpawnOwnership::RootSequence),
+			FCanExecuteAction(),
+			FIsActionChecked::CreateLambda([=] { return Spawnable->GetSpawnOwnership() == ESpawnOwnership::RootSequence; })
+		),
+		NAME_None,
+		EUserInterfaceActionType::ToggleButton
+	);
+
+	MenuBuilder.AddMenuEntry(
+		LOCTEXT("External_Label", "External"),
+		LOCTEXT("External_Tooltip", "Indicates this object's lifetime is managed externally once spawned. It will not be destroyed by sequencer."),
+		FSlateIcon(),
+		FUIAction(
+			FExecuteAction::CreateLambda(Callback, ESpawnOwnership::External),
+			FCanExecuteAction(),
+			FIsActionChecked::CreateLambda([=] { return Spawnable->GetSpawnOwnership() == ESpawnOwnership::External; })
+		),
+		NAME_None,
+		EUserInterfaceActionType::ToggleButton
+	);
+}
+
+void FLevelSequenceCustomization::AddSpawnLevelMenu(FMenuBuilder& MenuBuilder, TSharedPtr<FObjectBindingModel> ObjectBindingModel)
+{
+	TSharedPtr<ISequencer> Sequencer = WeakSequencer.Pin();
+	TSharedPtr<FSequencerEditorViewModel> EditorViewModel = Sequencer->GetViewModel();
+
+	FGuid ObjectBindingID = ObjectBindingModel->GetObjectGuid();
+	UMovieScene* MovieScene = Sequencer->GetFocusedMovieSceneSequence()->GetMovieScene();
+
+	FMovieSceneSpawnable* Spawnable = MovieScene->FindSpawnable(ObjectBindingID);
+	if (!Spawnable)
+	{
+		return;
+	}
+
+	MenuBuilder.AddMenuEntry(
+		NSLOCTEXT("UnrealEd", "PersistentLevel", "Persistent Level"),
+		NSLOCTEXT("UnrealEd", "PersistentLevel", "Persistent Level"),
+		FSlateIcon(),
+		FUIAction(
+			FExecuteAction::CreateRaw(this, &FLevelSequenceCustomization::SetSelectedNodesSpawnableLevel, ObjectBindingModel, FName(NAME_None)),
+			FCanExecuteAction(),
+			FIsActionChecked::CreateLambda([=] { return Spawnable->GetLevelName() == NAME_None; })
+		),
+		NAME_None,
+		EUserInterfaceActionType::ToggleButton
+	);
+
+	UWorld* World = Cast<UWorld>(Sequencer->GetPlaybackContext());
+	if (!World)
+	{
+		return;
+	}
+
+	for (ULevelStreaming* LevelStreaming : World->GetStreamingLevels())
+	{
+		if (LevelStreaming)
+		{
+			FName LevelName = FPackageName::GetShortFName(LevelStreaming->GetWorldAssetPackageFName());
+
+			MenuBuilder.AddMenuEntry(
+				FText::FromName(LevelName),
+				FText::FromName(LevelName),
+				FSlateIcon(),
+				FUIAction(
+					FExecuteAction::CreateRaw(this, &FLevelSequenceCustomization::SetSelectedNodesSpawnableLevel, ObjectBindingModel, LevelName),
+					FCanExecuteAction(),
+					FIsActionChecked::CreateLambda([=] { return Spawnable->GetLevelName() == LevelName; })
+				),
+				NAME_None,
+				EUserInterfaceActionType::ToggleButton
+			);
+		}
+	}
+}
+
+void FLevelSequenceCustomization::SetSelectedNodesSpawnableLevel(TSharedPtr<FObjectBindingModel> ObjectBindingModel, FName InLevelName)
+{
+	TSharedPtr<ISequencer> Sequencer = WeakSequencer.Pin();
+	TSharedPtr<FSequencerEditorViewModel> EditorViewModel = Sequencer->GetViewModel();
+	UMovieScene* MovieScene = Sequencer->GetFocusedMovieSceneSequence()->GetMovieScene();
+
+	if (MovieScene->IsReadOnly())
+	{
+		FSequencerUtilities::ShowReadOnlyError();
+		return;
+	}
+
+	const FScopedTransaction Transaction( LOCTEXT("SetSpawnableLevel", "Set Spawnable Level") );
+
+	MovieScene->Modify();
+
+	TArray<FMovieSceneSpawnable*> Spawnables;
+
+	for (TViewModelPtr<IObjectBindingExtension> ObjectBindingNode : EditorViewModel->GetSelection()->Outliner.Filter<IObjectBindingExtension>())
+	{
+		FMovieSceneSpawnable* Spawnable = MovieScene->FindSpawnable(ObjectBindingNode->GetObjectGuid());
+		if (Spawnable)
+		{
+			Spawnable->SetLevelName(InLevelName);
+		}
+	}
+}
+
+void FLevelSequenceCustomization::AddChangeClassMenu(FMenuBuilder& MenuBuilder, TSharedPtr<FObjectBindingModel> ObjectBindingModel)
+{
+	TSharedPtr<ISequencer> Sequencer = WeakSequencer.Pin();
+	TSharedPtr<FSequencerEditorViewModel> EditorViewModel = Sequencer->GetViewModel();
+
+	FGuid ObjectBindingID = ObjectBindingModel->GetObjectGuid();
+	UMovieScene* MovieScene = Sequencer->GetFocusedMovieSceneSequence()->GetMovieScene();
+
+	FMovieSceneSpawnable* Spawnable = MovieScene->FindSpawnable(ObjectBindingID);
+	if (!Spawnable)
+	{
+		return;
+	}
+
+	FClassViewerModule& ClassViewerModule = FModuleManager::LoadModuleChecked<FClassViewerModule>("ClassViewer");
+
+	FClassViewerInitializationOptions Options;
+	Options.Mode = EClassViewerMode::ClassPicker;
+	Options.bIsActorsOnly = true;
+	Options.bIsPlaceableOnly = true;
+
+	const UClass* ClassForObjectBinding = ObjectBindingModel->FindObjectClass();
+	if (ClassForObjectBinding)
+	{
+		Options.ViewerTitleString = FText::FromString(TEXT("Change from: ") + ClassForObjectBinding->GetFName().ToString());
+	}
+	else
+	{
+		Options.ViewerTitleString = FText::FromString(TEXT("Change from: (empty)"));
+	}
+
+	MenuBuilder.AddWidget(
+		SNew(SBox)
+		.MinDesiredWidth(300.0f)
+		.MaxDesiredHeight(400.0f)
+		[
+			ClassViewerModule.CreateClassViewer(Options, FOnClassPicked::CreateRaw(
+				this, &FLevelSequenceCustomization::HandleTemplateActorClassPicked, ObjectBindingModel))
+		],
+		FText(), true, false
+	);
+}
+
+void FLevelSequenceCustomization::HandleTemplateActorClassPicked(UClass* ChosenClass, TSharedPtr<FObjectBindingModel> ObjectBindingModel)
+{
+	TSharedPtr<ISequencer> Sequencer = WeakSequencer.Pin();
+	UMovieScene* MovieScene = Sequencer->GetFocusedMovieSceneSequence()->GetMovieScene();
+
+	FGuid ObjectBindingID = ObjectBindingModel->GetObjectGuid();
+
+	FSlateApplication::Get().DismissAllMenus();
+
+	FMovieSceneSpawnable* Spawnable = MovieScene->FindSpawnable(ObjectBindingID);
+	if (!Spawnable)
+	{
+		return;
+	}
+
+	FScopedTransaction Transaction(LOCTEXT("ChangeClass", "Change Class"));
+
+	MovieScene->Modify();
+
+	TValueOrError<FNewSpawnable, FText> Result = Sequencer->GetSpawnRegister().CreateNewSpawnableType(*ChosenClass, *MovieScene, nullptr);
+	if (Result.IsValid())
+	{
+		Spawnable->SetObjectTemplate(Result.GetValue().ObjectTemplate);
+
+		Sequencer->GetSpawnRegister().DestroySpawnedObject(Spawnable->GetGuid(), Sequencer->GetFocusedTemplateID(), *Sequencer.Get());
+		Sequencer->ForceEvaluate();
+	}
+}
+
+void FLevelSequenceCustomization::AddDynamicSpawnMenu(FMenuBuilder& MenuBuilder, TSharedPtr<FObjectBindingModel> ObjectBindingModel)
+{
+	TSharedPtr<ISequencer> Sequencer = WeakSequencer.Pin();
+	UMovieScene* MovieScene = Sequencer->GetFocusedMovieSceneSequence()->GetMovieScene();
+
+	FGuid ObjectBindingID = ObjectBindingModel->GetObjectGuid();
+	FMovieSceneSpawnable* Spawnable = MovieScene->FindSpawnable(ObjectBindingID);
+	if (!Spawnable)
+	{
+		return;
+	}
+
+	ObjectBindingModel->AddDynamicBindingMenu(MenuBuilder, Spawnable->DynamicBinding);
+}
+
+void FLevelSequenceCustomization::AddDynamicPossessionMenu(FMenuBuilder& MenuBuilder, TSharedPtr<FObjectBindingModel> ObjectBindingModel)
+{
+	TSharedPtr<ISequencer> Sequencer = WeakSequencer.Pin();
+	UMovieScene* MovieScene = Sequencer->GetFocusedMovieSceneSequence()->GetMovieScene();
+
+	FGuid ObjectBindingID = ObjectBindingModel->GetObjectGuid();
+	FMovieScenePossessable* Possessable = MovieScene->FindPossessable(ObjectBindingID);
+	if (!Possessable)
+	{
+		return;
+	}
+
+	ObjectBindingModel->AddDynamicBindingMenu(MenuBuilder, Possessable->DynamicBinding);
+}
+
+} // namespace UE::Sequencer
+
+#undef LOCTEXT_NAMESPACE
+
diff --git a/Engine/Plugins/MovieScene/LevelSequenceEditor/Source/LevelSequenceEditor/Private/Misc/LevelSequenceCustomization.h b/Engine/Plugins/MovieScene/LevelSequenceEditor/Source/LevelSequenceEditor/Private/Misc/LevelSequenceCustomization.h
new file mode 100644
--- /dev/null
+++ b/Engine/Plugins/MovieScene/LevelSequenceEditor/Source/LevelSequenceEditor/Private/Misc/LevelSequenceCustomization.h
@@ -0,0 +1,55 @@
+// Copyright Epic Games, Inc. All Rights Reserved.
+
+#pragma once
+
+#include "CoreTypes.h"
+#include "MVVM/ViewModelPtr.h"
+#include "SequencerCustomizationManager.h"
+#include "Toolkits/AssetEditorToolkit.h"
+
+namespace UE::Sequencer
+{
+	class FObjectBindingModel;
+	class FSequencerEditorViewModel;
+
+/**
+ * The sequencer customization for level sequences. 
+ */
+class FLevelSequenceCustomization : public ISequencerCustomization
+{
+protected:
+
+	//~ ISequencerCustomization interface
+	virtual void RegisterSequencerCustomization(FSequencerCustomizationBuilder& Builder) override;
+	virtual void UnregisterSequencerCustomization() override;
+
+private:
+
+	// Actions menu extensions
+	void ExtendActionsMenu(FMenuBuilder& MenuBuilder);
+	void OnSaveMovieSceneAsClicked();
+	void ImportFBX();
+	void ExportFBX();
+
+	// Object binding context menu extensions
+	TSharedPtr<FExtender> CreateObjectBindingContextMenuExtender(FViewModelPtr InViewModel);
+	void ExtendObjectBindingContextMenu(FMenuBuilder& MenuBuilder, TSharedPtr<FObjectBindingModel> ObjectBindingModel);
+	void AddSpawnOwnershipMenu(FMenuBuilder& MenuBuilder, TSharedPtr<FObjectBindingModel> ObjectBindingModel);
+	void AddSpawnLevelMenu(FMenuBuilder& MenuBuilder, TSharedPtr<FObjectBindingModel> ObjectBindingModel);
+	void SetSelectedNodesSpawnableLevel(TSharedPtr<FObjectBindingModel> ObjectBindingModel, FName InLevelName);
+	void AddChangeClassMenu(FMenuBuilder& MenuBuilder, TSharedPtr<FObjectBindingModel> ObjectBindingModel);
+	void HandleTemplateActorClassPicked(UClass* ChosenClass, TSharedPtr<FObjectBindingModel> ObjectBindingModel);
+
+	void AddDynamicSpawnMenu(FMenuBuilder& MenuBuilder, TSharedPtr<FObjectBindingModel> ObjectBindingModel);
+	void AddDynamicPossessionMenu(FMenuBuilder& MenuBuilder, TSharedPtr<FObjectBindingModel> ObjectBindingModel);
+
+private:
+
+	TWeakPtr<ISequencer> WeakSequencer;
+
+	TSharedPtr<FUICommandList> ActionsMenuCommandList;
+	TSharedPtr<FExtender> ActionsMenuExtender;
+};
+
+} // namespace UE::Sequencer
+
diff --git a/Engine/Plugins/MovieScene/LevelSequenceEditor/Source/LevelSequenceEditor/Private/Misc/LevelSequenceFBXInterop.cpp b/Engine/Plugins/MovieScene/LevelSequenceEditor/Source/LevelSequenceEditor/Private/Misc/LevelSequenceFBXInterop.cpp
new file mode 100644
--- /dev/null
+++ b/Engine/Plugins/MovieScene/LevelSequenceEditor/Source/LevelSequenceEditor/Private/Misc/LevelSequenceFBXInterop.cpp
@@ -0,0 +1,247 @@
+// Copyright Epic Games, Inc. All Rights Reserved.
+
+#include "Misc/LevelSequenceFBXInterop.h"
+
+#include "DesktopPlatformModule.h"
+#include "EditorDirectories.h"
+#include "Exporters/Exporter.h"
+#include "FbxExporter.h"
+#include "Framework/Application/SlateApplication.h"
+#include "Framework/Notifications/NotificationManager.h"
+#include "ISequencer.h"
+#include "MovieSceneToolHelpers.h"
+#include "MVVM/Extensions/IObjectBindingExtension.h"
+#include "MVVM/Extensions/IOutlinerExtension.h"
+#include "MVVM/Extensions/ITrackExtension.h"
+#include "MVVM/Selection/Selection.h"
+#include "MVVM/ViewModelPtr.h"
+#include "MVVM/ViewModels/ObjectBindingModel.h"
+#include "MVVM/ViewModels/SequencerEditorViewModel.h"
+#include "MVVM/ViewModels/OutlinerViewModel.h"
+#include "MVVM/ViewModels/ViewModelIterators.h"
+#include "SequencerExportTask.h"
+#include "SequencerUtilities.h"
+#include "UObject/UObjectIterator.h"
+#include "UnrealExporter.h"
+#include "Widgets/Notifications/SNotificationList.h"
+
+#define LOCTEXT_NAMESPACE "LevelSequenceFBXInterop"
+
+FLevelSequenceFBXInterop::FLevelSequenceFBXInterop(TSharedPtr<ISequencer> InSequencer)
+	: Sequencer(InSequencer)
+{
+}
+
+void FLevelSequenceFBXInterop::ImportFBX()
+{
+	using namespace UE::Sequencer;
+
+	TMap<FGuid, FString> ObjectBindingNameMap;
+
+	TSharedPtr<FSequencerEditorViewModel> EditorViewModel = Sequencer->GetViewModel();
+	TSharedPtr<FOutlinerViewModel> OutlinerViewModel = EditorViewModel->GetOutliner();
+	TParentFirstChildIterator<IObjectBindingExtension> ObjectBindingIt = OutlinerViewModel->GetRootItem()->GetDescendantsOfType<IObjectBindingExtension>();
+
+	// Only visit the first object binding in a given sub-hierarchy so we get top-level object bindings.
+	// This is done by skipping the branch on each iteration
+	for ( ; ObjectBindingIt; ++ObjectBindingIt)
+	{
+		FGuid ObjectBinding = ObjectBindingIt->GetObjectGuid();
+		ObjectBindingNameMap.Add(ObjectBinding, (*ObjectBindingIt).AsModel()->CastThisChecked<IOutlinerExtension>()->GetLabel().ToString());
+
+		ObjectBindingIt.IgnoreCurrentChildren();
+	}
+
+	MovieSceneToolHelpers::ImportFBXWithDialog(Sequencer->GetFocusedMovieSceneSequence(), *Sequencer, ObjectBindingNameMap, TOptional<bool>());
+}
+
+void FLevelSequenceFBXInterop::ImportFBXOntoSelectedNodes()
+{
+	using namespace UE::Sequencer;
+
+	// The object binding and names to match when importing from fbx
+	TMap<FGuid, FString> ObjectBindingNameMap;
+
+	TSharedPtr<FSequencerEditorViewModel> EditorViewModel = Sequencer->GetViewModel();
+	for (TViewModelPtr<FObjectBindingModel> ObjectBindingNode : EditorViewModel->GetSelection()->Outliner.Filter<FObjectBindingModel>())
+	{
+		FGuid ObjectBinding = ObjectBindingNode->GetObjectGuid();
+		ObjectBindingNameMap.Add(ObjectBinding, ObjectBindingNode->GetLabel().ToString());
+	}
+
+	MovieSceneToolHelpers::ImportFBXWithDialog(Sequencer->GetFocusedMovieSceneSequence(), *Sequencer, ObjectBindingNameMap, TOptional<bool>(false));
+}
+
+void FLevelSequenceFBXInterop::ExportFBX()
+{
+	using namespace UE::Sequencer;
+
+	TArray<UExporter*> Exporters;
+	TArray<FString> SaveFilenames;
+	IDesktopPlatform* DesktopPlatform = FDesktopPlatformModule::Get();
+	bool bExportFileNamePicked = false;
+	if ( DesktopPlatform != NULL )
+	{
+		FString FileTypes = "FBX document|*.fbx";
+		UMovieSceneSequence* Sequence = Sequencer->GetFocusedMovieSceneSequence();
+		for (TObjectIterator<UClass> It; It; ++It)
+		{
+			if (!It->IsChildOf(UExporter::StaticClass()) || It->HasAnyClassFlags(CLASS_Abstract | CLASS_Deprecated | CLASS_NewerVersionExists))
+			{
+				continue;
+			}
+
+			UExporter* Default = It->GetDefaultObject<UExporter>();
+			if (!Default->SupportsObject(Sequence))
+			{
+				continue;
+			}
+
+			for (int32 i = 0; i < Default->FormatExtension.Num(); ++i)
+			{
+				const FString& FormatExtension = Default->FormatExtension[i];
+				const FString& FormatDescription = Default->FormatDescription[i];
+
+				if (FileTypes.Len() > 0)
+				{
+					FileTypes += TEXT("|");
+				}
+				FileTypes += FormatDescription;
+				FileTypes += TEXT("|*.");
+				FileTypes += FormatExtension;
+			}
+
+			Exporters.Add(Default);
+		}
+
+		bExportFileNamePicked = DesktopPlatform->SaveFileDialog(
+			FSlateApplication::Get().FindBestParentWindowHandleForDialogs(nullptr),
+			LOCTEXT( "ExportLevelSequence", "Export Level Sequence" ).ToString(),
+			*( FEditorDirectories::Get().GetLastDirectory( ELastDirectory::FBX ) ),
+			TEXT( "" ),
+			*FileTypes,
+			EFileDialogFlags::None,
+			SaveFilenames );
+	}
+
+	if ( bExportFileNamePicked )
+	{
+		FString ExportFilename = SaveFilenames[0];
+		FEditorDirectories::Get().SetLastDirectory( ELastDirectory::FBX, FPaths::GetPath( ExportFilename ) ); // Save path as default for next time.
+
+		// Make sure external selection is up to date since export could happen on tracks that have been right clicked but not have their underlying bound objects selected yet since that happens on mouse up.
+		FSequencerUtilities::SynchronizeExternalSelectionWithSequencerSelection(Sequencer.ToSharedRef());
+		
+		// Select selected nodes if there are selected nodes
+		TArray<FGuid> Bindings;
+		TArray<UMovieSceneTrack*> Tracks;
+		UMovieScene* MovieScene = Sequencer->GetFocusedMovieSceneSequence()->GetMovieScene();
+		TSharedPtr<FSequencerEditorViewModel> EditorViewModel = Sequencer->GetViewModel();
+		for (FViewModelPtr Node : EditorViewModel->GetSelection()->Outliner)
+		{
+			if (TViewModelPtr<IObjectBindingExtension> ObjectBindingNode = Node.ImplicitCast())
+			{
+				Bindings.Add(ObjectBindingNode->GetObjectGuid());
+
+				for (TViewModelPtr<IOutlinerExtension> DescendantNode : Node->GetDescendantsOfType<IOutlinerExtension>())
+				{
+					if (!EditorViewModel->GetSelection()->Outliner.IsSelected(DescendantNode) && DescendantNode.AsModel()->IsA<IObjectBindingExtension>())
+					{
+						IObjectBindingExtension* DescendantObjectBindingNode = DescendantNode.AsModel()->CastThisChecked<IObjectBindingExtension>();
+						Bindings.Add(DescendantObjectBindingNode->GetObjectGuid());
+					}
+				}
+			}
+			else if (TViewModelPtr<ITrackExtension> TrackNode = Node.ImplicitCast())
+			{
+				UMovieSceneTrack* Track = TrackNode->GetTrack();
+				if (Track && MovieScene->ContainsTrack(*Track))
+				{
+					Tracks.Add(Track);
+				}
+			}
+		}
+
+		FString FileExtension = FPaths::GetExtension(ExportFilename);
+		if (FileExtension == TEXT("fbx"))
+		{
+			ExportFBXInternal(ExportFilename, Bindings, (Bindings.Num() + Tracks.Num()) > 0 ? Tracks : MovieScene->GetTracks());
+		}
+		else
+		{
+			for (UExporter* Exporter : Exporters)
+			{
+				if (Exporter->FormatExtension.Contains(FileExtension))
+				{
+					USequencerExportTask* ExportTask = NewObject<USequencerExportTask>();
+					TStrongObjectPtr<USequencerExportTask> ExportTaskGuard(ExportTask);
+					ExportTask->Object = Sequencer->GetFocusedMovieSceneSequence();
+					ExportTask->Exporter = nullptr;
+					ExportTask->Filename = ExportFilename;
+					ExportTask->bSelected = false;
+					ExportTask->bReplaceIdentical = true;
+					ExportTask->bPrompt = false;
+					ExportTask->bUseFileArchive = false;
+					ExportTask->bWriteEmptyFiles = false;
+					ExportTask->bAutomated = false;
+					ExportTask->Exporter = NewObject<UExporter>(GetTransientPackage(), Exporter->GetClass());
+
+					ExportTask->SequencerContext = Sequencer->GetPlaybackContext();
+
+					UExporter::RunAssetExportTask(ExportTask);
+
+					ExportTask->Object = nullptr;
+					ExportTask->Exporter = nullptr;
+					ExportTask->SequencerContext = nullptr;
+
+					break;
+				}
+			}
+		}
+	}
+}
+
+void FLevelSequenceFBXInterop::ExportFBXInternal(const FString& ExportFilename, const TArray<FGuid>& Bindings, const TArray<UMovieSceneTrack*>& Tracks)
+{
+	UnFbx::FFbxExporter* Exporter = UnFbx::FFbxExporter::GetInstance();
+	//Show the fbx export dialog options
+	bool ExportCancel = false;
+	bool ExportAll = false;
+	Exporter->FillExportOptions(false, true, ExportFilename, ExportCancel, ExportAll);
+	if (!ExportCancel)
+	{
+		UMovieScene* MovieScene = Sequencer->GetFocusedMovieSceneSequence()->GetMovieScene();
+		UWorld* World = Cast<UWorld>(Sequencer->GetPlaybackContext());
+		FMovieSceneSequenceIDRef Template = Sequencer->GetFocusedTemplateID();
+		UnFbx::FFbxExporter::FLevelSequenceNodeNameAdapter NodeNameAdapter(MovieScene, Sequencer.Get(), Template);
+
+		{
+			FSpawnableRestoreState SpawnableRestoreState(MovieScene);
+			if (SpawnableRestoreState.bWasChanged)
+			{
+				// Evaluate at the beginning of the subscene time to ensure that spawnables are created before export
+				Sequencer->SetLocalTimeDirectly(UE::MovieScene::DiscreteInclusiveLower(FSequencerUtilities::GetTimeBounds(Sequencer.ToSharedRef())));
+			}
+
+			FMovieSceneSequenceTransform RootToLocalTransform = Sequencer->GetFocusedMovieSceneSequenceTransform();
+			if (MovieSceneToolHelpers::ExportFBX(World, MovieScene, Sequencer.Get(), Bindings, Tracks, NodeNameAdapter, Template, ExportFilename, RootToLocalTransform))
+			{
+				FNotificationInfo Info(NSLOCTEXT("Sequencer", "ExportFBXSucceeded", "FBX Export Succeeded."));
+				Info.Hyperlink = FSimpleDelegate::CreateStatic([](FString InFilename) { FPlatformProcess::ExploreFolder(*InFilename); }, ExportFilename);
+				Info.HyperlinkText = FText::FromString(ExportFilename);
+				Info.ExpireDuration = 5.0f;
+				FSlateNotificationManager::Get().AddNotification(Info)->SetCompletionState(SNotificationItem::CS_Success);
+			}
+			else
+			{
+				FNotificationInfo Info(NSLOCTEXT("Sequencer", "ExportFBXFailed", "FBX Export Failed."));
+				Info.ExpireDuration = 5.0f;
+				FSlateNotificationManager::Get().AddNotification(Info)->SetCompletionState(SNotificationItem::CS_Fail);
+			}
+		}
+
+		Sequencer->ForceEvaluate();
+	}
+}
+
+#undef LOCTEXT_NAMESPACE
diff --git a/Engine/Plugins/MovieScene/LevelSequenceEditor/Source/LevelSequenceEditor/Private/Misc/LevelSequenceFBXInterop.h b/Engine/Plugins/MovieScene/LevelSequenceEditor/Source/LevelSequenceEditor/Private/Misc/LevelSequenceFBXInterop.h
new file mode 100644
--- /dev/null
+++ b/Engine/Plugins/MovieScene/LevelSequenceEditor/Source/LevelSequenceEditor/Private/Misc/LevelSequenceFBXInterop.h
@@ -0,0 +1,29 @@
+// Copyright Epic Games, Inc. All Rights Reserved.
+
+#pragma once
+
+#include "ISequencer.h"
+
+class FLevelSequenceFBXInterop
+{
+public:
+
+	FLevelSequenceFBXInterop(TSharedPtr<ISequencer> InSequencer);
+
+	/** Imports the animation from an fbx file. */
+	void ImportFBX();
+	void ImportFBXOntoSelectedNodes();
+
+	/** Exports the animation to an fbx file. */
+	void ExportFBX();
+
+private:
+
+	/** Exports sequence to a FBX file */
+	void ExportFBXInternal(const FString& ExportFilename, const TArray<FGuid>& Bindings, const TArray<UMovieSceneTrack*>& Tracks);
+
+private:
+
+	TSharedPtr<ISequencer> Sequencer;
+};
+
diff --git a/Engine/Plugins/MovieScene/LevelSequenceEditor/Source/LevelSequenceEditor/Private/Styles/LevelSequenceEditorStyle.h b/Engine/Plugins/MovieScene/LevelSequenceEditor/Source/LevelSequenceEditor/Private/Styles/LevelSequenceEditorStyle.h
--- a/Engine/Plugins/MovieScene/LevelSequenceEditor/Source/LevelSequenceEditor/Private/Styles/LevelSequenceEditorStyle.h
+++ b/Engine/Plugins/MovieScene/LevelSequenceEditor/Source/LevelSequenceEditor/Private/Styles/LevelSequenceEditorStyle.h
@@ -24,6 +24,7 @@ public:
 		const FVector2D Icon16x16(16.0f, 16.0f);
 		const FVector2D Icon20x20(20.0f, 20.0f);
 		const FVector2D Icon40x40(40.0f, 40.0f);
+		const FVector2D Icon48x48(48.0f, 48.0f);
 
 		SetContentRoot(FPaths::EnginePluginsDir() / TEXT("MovieScene/LevelSequenceEditor/Content"));
 		SetCoreContentRoot(FPaths::EngineContentDir() / TEXT("Slate"));
@@ -46,6 +47,10 @@ public:
 
 		Set("LevelSequenceEditor.CinematicViewportTransportRangeKey", new IMAGE_BRUSH("CinematicViewportTransportRangeKey", FVector2D(7.f, 7.f)));
 
+		Set("LevelSequenceEditor.SaveAs", new IMAGE_BRUSH("Icon_Sequencer_SaveAs_48x", Icon48x48));
+		Set("LevelSequenceEditor.ImportFBX", new IMAGE_BRUSH("Icon_Sequencer_ImportFBX_48x", Icon48x48));
+		Set("LevelSequenceEditor.ExportFBX", new IMAGE_BRUSH("Icon_Sequencer_ExportFBX_48x", Icon48x48));
+
 		Set("FilmOverlay.DefaultThumbnail", new IMAGE_BRUSH("DefaultFilmOverlayThumbnail", FVector2D(36, 24)));
 
 		Set("FilmOverlay.Disabled", new IMAGE_BRUSH("FilmOverlay.Disabled", FVector2D(36, 24)));
diff --git a/Engine/Plugins/MovieScene/LevelSequenceEditor/Source/LevelSequenceEditor/Public/SequencerExportTask.h b/Engine/Plugins/MovieScene/LevelSequenceEditor/Source/LevelSequenceEditor/Public/SequencerExportTask.h
new file mode 100644
--- /dev/null
+++ b/Engine/Plugins/MovieScene/LevelSequenceEditor/Source/LevelSequenceEditor/Public/SequencerExportTask.h
@@ -0,0 +1,24 @@
+// Copyright Epic Games, Inc. All Rights Reserved.
+
+#pragma once
+
+#include "CoreMinimal.h"
+#include "UObject/ObjectMacros.h"
+#include "UObject/Object.h"
+#include "AssetExportTask.h"
+#include "SequencerExportTask.generated.h"
+
+/**
+ * Contains data for a group of assets to import
+ */ 
+UCLASS(Transient, BlueprintType)
+class USequencerExportTask : public UAssetExportTask
+{
+	GENERATED_BODY()
+
+public:
+
+	/* A UWorld for LevelSequences, UUserWidget for WidgetAnimations, or AActor for Actor Sequences, etc... */
+	UPROPERTY(BlueprintReadWrite, Category = Miscellaneous)
+	TObjectPtr<UObject> SequencerContext;
+};
diff --git a/Engine/Source/Editor/EditorStyle/Private/SlateEditorStyle.cpp b/Engine/Source/Editor/EditorStyle/Private/SlateEditorStyle.cpp
--- a/Engine/Source/Editor/EditorStyle/Private/SlateEditorStyle.cpp
+++ b/Engine/Source/Editor/EditorStyle/Private/SlateEditorStyle.cpp
@@ -2384,12 +2384,6 @@ void FSlateEditorStyle::FStyle::SetupSequencerStyles()
 		Set("Sequencer.Save.Small", new IMAGE_BRUSH("Sequencer/Main_Icons/Icon_Sequencer_Save_48x", Icon24x24));
 		Set("Sequencer.SaveAsterisk", new IMAGE_BRUSH("Sequencer/Main_Icons/Icon_Sequencer_SaveAsterisk_48x", Icon48x48));
 		Set("Sequencer.SaveAsterisk.Small", new IMAGE_BRUSH("Sequencer/Main_Icons/Icon_Sequencer_SaveAsterisk_48x", Icon24x24));
-		Set("Sequencer.SaveAs", new IMAGE_BRUSH("Sequencer/Main_Icons/Icon_Sequencer_SaveAs_48x", Icon48x48));
-		Set("Sequencer.SaveAs.Small", new IMAGE_BRUSH("Sequencer/Main_Icons/Icon_Sequencer_SaveAs_48x", Icon24x24));
-		Set("Sequencer.ImportFBX", new IMAGE_BRUSH("Sequencer/Main_Icons/Icon_Sequencer_ImportFBX_48x", Icon48x48));
-		Set("Sequencer.ImportFBX.Small", new IMAGE_BRUSH("Sequencer/Main_Icons/Icon_Sequencer_ImportFBX_48x", Icon24x24));
-		Set("Sequencer.ExportFBX", new IMAGE_BRUSH("Sequencer/Main_Icons/Icon_Sequencer_ExportFBX_48x", Icon48x48));
-		Set("Sequencer.ExportFBX.Small", new IMAGE_BRUSH("Sequencer/Main_Icons/Icon_Sequencer_ExportFBX_48x", Icon24x24));
 		Set("Sequencer.DiscardChanges", new IMAGE_BRUSH("Sequencer/Main_Icons/Icon_Sequencer_Revert_24x", Icon48x48));
 		Set("Sequencer.DiscardChanges.Small", new IMAGE_BRUSH("Sequencer/Main_Icons/Icon_Sequencer_Revert_24x", Icon24x24));
 		Set("Sequencer.RestoreAnimatedState", new IMAGE_BRUSH("Sequencer/Main_Icons/Icon_Sequencer_RestoreAnimatedState_24x", Icon48x48));
@@ -2656,6 +2650,7 @@ void FSlateEditorStyle::FStyle::SetupSequencerStyles()
 	Set( "Sequencer.ToolBar.RadioButton", 				FCoreStyle::Get().GetWidgetStyle<FCheckBoxStyle>("ToolBar.RadioButton"));
 	Set( "Sequencer.ToolBar.ToggleButton", 				FCoreStyle::Get().GetWidgetStyle<FCheckBoxStyle>("ToolBar.ToggleButton"));
 	Set( "Sequencer.ToolBar.Button", 					FCoreStyle::Get().GetWidgetStyle<FButtonStyle>("ToolBar.Button"));
+
 	Set( "ToolBar.Button.Normal", 						new FSlateNoResource() );
 	Set( "ToolBar.Button.Pressed", 						new CORE_BOX_BRUSH( "Common/RoundedSelection_16x", 4.0f/16.0f, SelectionColor_Pressed ) );
 	Set( "ToolBar.Button.Hovered", 						new CORE_BOX_BRUSH( "Common/RoundedSelection_16x", 4.0f/16.0f, SelectionColor ) );
diff --git a/Engine/Source/Editor/EditorStyle/Private/StarshipStyle.cpp b/Engine/Source/Editor/EditorStyle/Private/StarshipStyle.cpp
--- a/Engine/Source/Editor/EditorStyle/Private/StarshipStyle.cpp
+++ b/Engine/Source/Editor/EditorStyle/Private/StarshipStyle.cpp
@@ -2130,9 +2130,6 @@ void FStarshipEditorStyle::FStyle::SetupSequencerStyles()
 
 		Set("Sequencer.NotificationImage_AddedPlayMovieSceneEvent", new IMAGE_BRUSH("Old/Checkbox_checked", Icon16x16));
 
-		Set("Sequencer.SaveAs", new IMAGE_BRUSH("Sequencer/Main_Icons/Icon_Sequencer_SaveAs_48x", Icon48x48));
-		Set("Sequencer.ImportFBX", new IMAGE_BRUSH("Sequencer/Main_Icons/Icon_Sequencer_ImportFBX_48x", Icon48x48));
-		Set("Sequencer.ExportFBX", new IMAGE_BRUSH("Sequencer/Main_Icons/Icon_Sequencer_ExportFBX_48x", Icon48x48));
 		Set("Sequencer.RestoreAnimatedState", new IMAGE_BRUSH("Sequencer/Main_Icons/Icon_Sequencer_RestoreAnimatedState_24x", Icon48x48));
 		Set("Sequencer.GenericGripLeft", new BOX_BRUSH("Sequencer/GenericGripLeft", FMargin(5.f / 16.f)));
 		Set("Sequencer.GenericGripRight", new BOX_BRUSH("Sequencer/GenericGripRight", FMargin(5.f / 16.f)));
diff --git a/Engine/Source/Editor/Sequencer/Private/MVVM/ViewModels/ObjectBindingModel.cpp b/Engine/Source/Editor/Sequencer/Private/MVVM/ViewModels/ObjectBindingModel.cpp
--- a/Engine/Source/Editor/Sequencer/Private/MVVM/ViewModels/ObjectBindingModel.cpp
+++ b/Engine/Source/Editor/Sequencer/Private/MVVM/ViewModels/ObjectBindingModel.cpp
@@ -42,7 +42,6 @@
 #include "SequencerUtilities.h"
 #include "Styling/AppStyle.h"
 #include "Styling/SlateIconFinder.h"
-#include "Tracks/MovieSceneSpawnTrack.h"
 
 #define LOCTEXT_NAMESPACE "ObjectBindingModel"
 
@@ -118,60 +117,6 @@ void GetKeyablePropertyPaths(UClass* Class, void* ValuePtr, UStruct* PropertySou
 	}
 }
 
-struct FMovieSceneSpawnableFlagCheckState
-{
-	FSequencer* Sequencer;
-	UMovieScene* MovieScene;
-	bool FMovieSceneSpawnable::*PtrToFlag;
-
-	ECheckBoxState operator()() const
-	{
-		using namespace UE::Sequencer;
-
-		ECheckBoxState CheckState = ECheckBoxState::Undetermined;
-		for (TViewModelPtr<IObjectBindingExtension> ObjectBinding : Sequencer->GetViewModel()->GetSelection()->Outliner.Filter<IObjectBindingExtension>())
-		{
-			FMovieSceneSpawnable* SelectedSpawnable = MovieScene->FindSpawnable(ObjectBinding->GetObjectGuid());
-			if (SelectedSpawnable)
-			{
-				if (CheckState != ECheckBoxState::Undetermined && SelectedSpawnable->*PtrToFlag != ( CheckState == ECheckBoxState::Checked ))
-				{
-					return ECheckBoxState::Undetermined;
-				}
-				CheckState = SelectedSpawnable->*PtrToFlag ? ECheckBoxState::Checked : ECheckBoxState::Unchecked;
-			}
-		}
-		return CheckState;
-	}
-};
-
-struct FMovieSceneSpawnableFlagToggler
-{
-	FSequencer* Sequencer;
-	UMovieScene* MovieScene;
-	bool FMovieSceneSpawnable::*PtrToFlag;
-	FText TransactionText;
-
-	void operator()() const
-	{
-		using namespace UE::Sequencer;
-
-		FScopedTransaction Transaction(TransactionText);
-
-		const ECheckBoxState CheckState = FMovieSceneSpawnableFlagCheckState{Sequencer, MovieScene, PtrToFlag}();
-
-		MovieScene->Modify();
-		for (TViewModelPtr<IObjectBindingExtension> ObjectBinding : Sequencer->GetViewModel()->GetSelection()->Outliner.Filter<IObjectBindingExtension>())
-		{
-			FMovieSceneSpawnable* SelectedSpawnable = MovieScene->FindSpawnable(ObjectBinding->GetObjectGuid());
-			if (SelectedSpawnable)
-			{
-				SelectedSpawnable->*PtrToFlag = (CheckState == ECheckBoxState::Unchecked);
-			}
-		}
-	}
-};
-
 } // anon-namespace
 
 FObjectBindingModel::FObjectBindingModel(FSequenceModel* InOwnerModel, const FMovieSceneBinding& InBinding)
@@ -912,117 +857,29 @@ void FObjectBindingModel::AddPropertyMenuItems(FMenuBuilder& AddTrackMenuBuilder
 
 void FObjectBindingModel::BuildContextMenu(FMenuBuilder& MenuBuilder)
 {
-	FSequencer* Sequencer = GetEditor()->GetSequencerImpl().Get();
+	TSharedPtr<FSequencerEditorViewModel> EditorViewModel = GetEditor();
+	FSequencer* Sequencer = EditorViewModel->GetSequencerImpl().Get();
 	ISequencerModule& SequencerModule = FModuleManager::GetModuleChecked<ISequencerModule>("Sequencer");
 
 	UObject* BoundObject = Sequencer->FindSpawnedObjectOrTemplate(ObjectBindingID);
 	const UClass* ObjectClass = FindObjectClass();
-
-	TSharedRef<FUICommandList> CommandList(new FUICommandList);
-	TSharedPtr<FExtender> Extender = SequencerModule.GetObjectBindingContextMenuExtensibilityManager()->GetAllExtenders(CommandList, TArrayBuilder<UObject*>().Add(BoundObject));
+	
+	TSharedPtr<FExtender> Extender = EditorViewModel->GetSequencerMenuExtender(
+			SequencerModule.GetObjectBindingContextMenuExtensibilityManager(), TArrayBuilder<UObject*>().Add(BoundObject),
+			&FSequencerCustomizationInfo::OnBuildObjectBindingContextMenu, SharedThis(this));
 	if (Extender.IsValid())
 	{
 		MenuBuilder.PushExtender(Extender.ToSharedRef());
 	}
-
-	if (Sequencer->IsLevelEditorSequencer())
-	{
-		UMovieScene* MovieScene = OwnerModel->GetMovieScene();
-		FMovieSceneSpawnable* Spawnable = MovieScene->FindSpawnable(ObjectBindingID);
-
-		if (Spawnable)
-		{
-			MenuBuilder.BeginSection("Spawnable", LOCTEXT("SpawnableMenuSectionName", "Spawnable"));
 	
-			MenuBuilder.AddSubMenu(
-				LOCTEXT("OwnerLabel", "Spawned Object Owner"),
-				LOCTEXT("OwnerTooltip", "Specifies how the spawned object is to be owned"),
-				FNewMenuDelegate::CreateSP(this, &FObjectBindingModel::AddSpawnOwnershipMenu)
-			);
-
-			MenuBuilder.AddSubMenu(
-				LOCTEXT("SubLevelLabel", "Spawnable Level"),
-				LOCTEXT("SubLevelTooltip", "Specifies which level the spawnable should be spawned into"),
-				FNewMenuDelegate::CreateSP(this, &FObjectBindingModel::AddSpawnLevelMenu)
-			);
-
-			MenuBuilder.AddSubMenu(
-				LOCTEXT("ChangeClassLabel", "Change Class"),
-				LOCTEXT("ChangeClassTooltip", "Change the class (object template) that this spawns from"),
-				FNewMenuDelegate::CreateSP(this, &FObjectBindingModel::AddChangeClassMenu));
-
-			MenuBuilder.AddSubMenu(
-				LOCTEXT("DynamicSpawn", "Dynamic Spawn"),
-				LOCTEXT("DynamicSpawnTooltip", "Specify a Blueprint method that will spawn or otherwise acquire a compatible actor for this binding"),
-				FNewMenuDelegate::CreateSP(this, &FObjectBindingModel::AddDynamicSpawnMenu));
-
-			MenuBuilder.AddMenuEntry(
-				LOCTEXT("ContinuouslyRespawn", "Continuously Respawn"),
-				LOCTEXT("ContinuouslyRespawnTooltip", "When enabled, this spawnable will always be respawned if it gets destroyed externally. When disabled, this object will only ever be spawned once for each spawn key even if it is destroyed externally"),
-				FSlateIcon(),
-				FUIAction(
-					FExecuteAction::CreateLambda(FMovieSceneSpawnableFlagToggler{Sequencer, MovieScene, &FMovieSceneSpawnable::bContinuouslyRespawn, LOCTEXT("ContinuouslyRespawnTransaction", "Set Continuously Respawn")}),
-					FCanExecuteAction(),
-					FGetActionCheckState::CreateLambda(FMovieSceneSpawnableFlagCheckState{Sequencer, MovieScene, &FMovieSceneSpawnable::bContinuouslyRespawn})
-				),
-				NAME_None,
-				EUserInterfaceActionType::ToggleButton
-			);
-
-			MenuBuilder.AddMenuEntry(
-				LOCTEXT("NetAddressable", "Net Addressable"),
-				LOCTEXT("NetAddressableTooltip", "When enabled, this spawnable will be spawned using a unique name that allows it to be addressed by the server and client (useful for relative movement calculations on spawned props)"),
-				FSlateIcon(),
-				FUIAction(
-					FExecuteAction::CreateLambda(FMovieSceneSpawnableFlagToggler{Sequencer, MovieScene, &FMovieSceneSpawnable::bNetAddressableName, LOCTEXT("NetAddressableTransaction", "Set Net Addressable")}),
-					FCanExecuteAction(),
-					FGetActionCheckState::CreateLambda(FMovieSceneSpawnableFlagCheckState{Sequencer, MovieScene, &FMovieSceneSpawnable::bNetAddressableName})
-				),
-				NAME_None,
-				EUserInterfaceActionType::ToggleButton
-			);
-
-			MenuBuilder.AddMenuEntry( FSequencerCommands::Get().SaveCurrentSpawnableState );
-			MenuBuilder.AddMenuEntry( FSequencerCommands::Get().ConvertToPossessable );
-
-			MenuBuilder.EndSection();
-		}
-		else
-		{
-			MenuBuilder.BeginSection("Possessable");
-
-			MenuBuilder.AddMenuEntry( FSequencerCommands::Get().ConvertToSpawnable );
-
-			MenuBuilder.AddSubMenu(
-				LOCTEXT("DynamicPossession", "Dynamic Possession"),
-				LOCTEXT("DynamicPossessionTooltip", "Specify a Blueprint method that will find a compatible possessable actor for this binding"),
-				FNewMenuDelegate::CreateSP(this, &FObjectBindingModel::AddDynamicPossessionMenu));
-
-			MenuBuilder.EndSection();
-		}
-
-		MenuBuilder.BeginSection("Import/Export", LOCTEXT("ImportExportMenuSectionName", "Import/Export"));
-		
-		MenuBuilder.AddMenuEntry(
-			LOCTEXT("ImportFBX", "Import..."),
-			LOCTEXT("ImportFBXTooltip", "Import FBX animation to this object"),
-			FSlateIcon(),
-			FUIAction(
-				FExecuteAction::CreateLambda([=]{ Sequencer->ImportFBXOntoSelectedNodes(); })
-			));
-
-		MenuBuilder.AddMenuEntry(
-			LOCTEXT("ExportFBX", "Export..."),
-			LOCTEXT("ExportFBXTooltip", "Export FBX animation from this object"),
-			FSlateIcon(),
-			FUIAction(
-				FExecuteAction::CreateLambda([=]{ Sequencer->ExportFBX(); })
-			));
-			
-		MenuBuilder.EndSection();
-	}
+	// Extenders can go in there.
+	MenuBuilder.BeginSection("ObjectBindingActions");
+	MenuBuilder.EndSection();
 
+	// External extension.
 	Sequencer->BuildCustomContextMenuForGuid(MenuBuilder, ObjectBindingID);
+
+	// Track editor extension.
 	TArray<FGuid> ObjectBindings;
 	ObjectBindings.Add(ObjectBindingID);
 	for (const TSharedPtr<ISequencerTrackEditor>& TrackEditor : Sequencer->GetTrackEditors())
@@ -1030,6 +887,7 @@ void FObjectBindingModel::BuildContextMenu(FMenuBuilder& MenuBuilder)
 		TrackEditor->BuildObjectBindingContextMenu(MenuBuilder, ObjectBindings, ObjectClass);
 	}
 
+	// Up-call.
 	FOutlinerItemModel::BuildContextMenu(MenuBuilder);
 }
 
@@ -1043,224 +901,7 @@ void FObjectBindingModel::BuildOrganizeContextMenu(FMenuBuilder& MenuBuilder)
 
 	FOutlinerItemModel::BuildOrganizeContextMenu(MenuBuilder);
 }
-
-void FObjectBindingModel::AddSpawnOwnershipMenu(FMenuBuilder& MenuBuilder)
-{
-	TSharedPtr<ISequencer> Sequencer = OwnerModel->GetSequencer();
-	UMovieScene* MovieScene = OwnerModel->GetMovieScene();
-	FMovieSceneSpawnable* Spawnable = MovieScene->FindSpawnable(ObjectBindingID);
-	if (!Spawnable)
-	{
-		return;
-	}
-	auto Callback = [=](ESpawnOwnership NewOwnership){
-
-		using namespace UE::Sequencer;
-
-		FScopedTransaction Transaction(LOCTEXT("SetSpawnOwnership", "Set Spawnable Ownership"));
-
-		MovieScene->Modify();
-
-		for (TViewModelPtr<IObjectBindingExtension> ObjectBinding : Sequencer->GetViewModel()->GetSelection()->Outliner.Filter<IObjectBindingExtension>())
-		{
-			FMovieSceneSpawnable* Spawnable = MovieScene->FindSpawnable(ObjectBinding->GetObjectGuid());
-		
-			Spawnable->SetSpawnOwnership(NewOwnership);
-
-			// Overwrite the completion state for all spawn sections to ensure the expected behaviour.
-			EMovieSceneCompletionMode NewCompletionMode = NewOwnership == ESpawnOwnership::InnerSequence ? EMovieSceneCompletionMode::RestoreState : EMovieSceneCompletionMode::KeepState;
-
-			// Make all spawn sections retain state
-			UMovieSceneSpawnTrack* SpawnTrack = MovieScene->FindTrack<UMovieSceneSpawnTrack>(ObjectBinding->GetObjectGuid());
-			if (SpawnTrack)
-			{
-				for (UMovieSceneSection* Section : SpawnTrack->GetAllSections())
-				{
-					Section->Modify();
-					Section->EvalOptions.CompletionMode = NewCompletionMode;
-				}
-			}
-		}
-	};
-
-	MenuBuilder.AddMenuEntry(
-		LOCTEXT("ThisSequence_Label", "This Sequence"),
-		LOCTEXT("ThisSequence_Tooltip", "Indicates that this sequence will own the spawned object. The object will be destroyed at the end of the sequence."),
-		FSlateIcon(),
-		FUIAction(
-			FExecuteAction::CreateLambda(Callback, ESpawnOwnership::InnerSequence),
-			FCanExecuteAction(),
-			FIsActionChecked::CreateLambda([=]{ return Spawnable->GetSpawnOwnership() == ESpawnOwnership::InnerSequence; })
-		),
-		NAME_None,
-		EUserInterfaceActionType::ToggleButton
-	);
-
-	MenuBuilder.AddMenuEntry(
-		LOCTEXT("RootSequence_Label", "Root Sequence"),
-		LOCTEXT("RootSequence_Tooltip", "Indicates that the outermost sequence will own the spawned object. The object will be destroyed when the outermost sequence stops playing."),
-		FSlateIcon(),
-		FUIAction(
-			FExecuteAction::CreateLambda(Callback, ESpawnOwnership::RootSequence),
-			FCanExecuteAction(),
-			FIsActionChecked::CreateLambda([=]{ return Spawnable->GetSpawnOwnership() == ESpawnOwnership::RootSequence; })
-		),
-		NAME_None,
-		EUserInterfaceActionType::ToggleButton
-	);
-
-	MenuBuilder.AddMenuEntry(
-		LOCTEXT("External_Label", "External"),
-		LOCTEXT("External_Tooltip", "Indicates this object's lifetime is managed externally once spawned. It will not be destroyed by sequencer."),
-		FSlateIcon(),
-		FUIAction(
-			FExecuteAction::CreateLambda(Callback, ESpawnOwnership::External),
-			FCanExecuteAction(),
-			FIsActionChecked::CreateLambda([=]{ return Spawnable->GetSpawnOwnership() == ESpawnOwnership::External; })
-		),
-		NAME_None,
-		EUserInterfaceActionType::ToggleButton
-	);
-}
-
-void FObjectBindingModel::AddSpawnLevelMenu(FMenuBuilder& MenuBuilder)
-{
-	TSharedPtr<FSequencer> Sequencer = OwnerModel->GetSequencerImpl();
-	UMovieScene* MovieScene = OwnerModel->GetMovieScene();
-	FMovieSceneSpawnable* Spawnable = MovieScene->FindSpawnable(ObjectBindingID);
-	if (!Spawnable)
-	{
-		return;
-	}
-
-	MenuBuilder.AddMenuEntry(
-		NSLOCTEXT("UnrealEd", "PersistentLevel", "Persistent Level"),
-		NSLOCTEXT("UnrealEd", "PersistentLevel", "Persistent Level"),
-		FSlateIcon(),
-		FUIAction(
-			FExecuteAction::CreateLambda([=] { Sequencer->SetSelectedNodesSpawnableLevel(NAME_None); }),
-			FCanExecuteAction(),
-			FIsActionChecked::CreateLambda([=] { return Spawnable->GetLevelName() == NAME_None; })
-		),
-		NAME_None,
-		EUserInterfaceActionType::ToggleButton
-	);
-
-	UWorld* World = Cast<UWorld>(Sequencer->GetPlaybackContext());
-	if (!World)
-	{
-		return;
-	}
-
-	for (ULevelStreaming* LevelStreaming : World->GetStreamingLevels())
-	{
-		if (LevelStreaming)
-		{
-			FName LevelName = FPackageName::GetShortFName( LevelStreaming->GetWorldAssetPackageFName() );
-
-			MenuBuilder.AddMenuEntry(
-				FText::FromName(LevelName),
-				FText::FromName(LevelName),
-				FSlateIcon(),
-				FUIAction(
-					FExecuteAction::CreateLambda([=] { Sequencer->SetSelectedNodesSpawnableLevel(LevelName); }),
-					FCanExecuteAction(),
-					FIsActionChecked::CreateLambda([=] { return Spawnable->GetLevelName() == LevelName; })
-				),
-				NAME_None,
-				EUserInterfaceActionType::ToggleButton
-			);
-		}
-	}
-}
-
-void FObjectBindingModel::AddChangeClassMenu(FMenuBuilder& MenuBuilder)
-{
-	UMovieScene* MovieScene = OwnerModel->GetMovieScene();
-	FMovieSceneSpawnable* Spawnable = MovieScene->FindSpawnable(ObjectBindingID);
-	if (!Spawnable)
-	{
-		return;
-	}
-
-	FClassViewerModule& ClassViewerModule = FModuleManager::LoadModuleChecked<FClassViewerModule>("ClassViewer");
-
-	FClassViewerInitializationOptions Options;
-	Options.Mode = EClassViewerMode::ClassPicker;
-	Options.bIsActorsOnly = true;
-	Options.bIsPlaceableOnly = true;
-
-	const UClass* ClassForObjectBinding = FindObjectClass();
-	if (ClassForObjectBinding)
-	{
-		Options.ViewerTitleString = FText::FromString(TEXT("Change from: ") + ClassForObjectBinding->GetFName().ToString());
-	}
-	else
-	{
-		Options.ViewerTitleString = FText::FromString(TEXT("Change from: (empty)"));
-	}
-
-	MenuBuilder.AddWidget(
-		SNew(SBox)
-		.MinDesiredWidth(300.0f)
-		.MaxDesiredHeight(400.0f)
-		[
-			ClassViewerModule.CreateClassViewer(Options, FOnClassPicked::CreateRaw(this, &FObjectBindingModel::HandleTemplateActorClassPicked))
-		],
-		FText(), true, false
-	);
-}
-
-void FObjectBindingModel::HandleTemplateActorClassPicked(UClass* ChosenClass)
-{
-	FSlateApplication::Get().DismissAllMenus();
-
-	TSharedPtr<ISequencer> Sequencer = OwnerModel->GetSequencer();
-	UMovieScene* MovieScene = OwnerModel->GetMovieScene();
-	FMovieSceneSpawnable* Spawnable = MovieScene->FindSpawnable(ObjectBindingID);
-	if (!Spawnable)
-	{
-		return;
-	}
-
-	FScopedTransaction Transaction(LOCTEXT("ChangeClass", "Change Class"));
-
-	MovieScene->Modify();
-
-	TValueOrError<FNewSpawnable, FText> Result = Sequencer->GetSpawnRegister().CreateNewSpawnableType(*ChosenClass, *MovieScene, nullptr);
-	if (Result.IsValid())
-	{
-		Spawnable->SetObjectTemplate(Result.GetValue().ObjectTemplate);
-
-		Sequencer->GetSpawnRegister().DestroySpawnedObject(Spawnable->GetGuid(), OwnerModel->GetSequenceID(), *Sequencer.Get());
-		Sequencer->ForceEvaluate();
-	}
-}
-
-void FObjectBindingModel::AddDynamicSpawnMenu(FMenuBuilder& MenuBuilder)
-{
-	UMovieScene* MovieScene = OwnerModel->GetMovieScene();
-	FMovieSceneSpawnable* Spawnable = MovieScene->FindSpawnable(ObjectBindingID);
-	if (!Spawnable)
-	{
-		return;
-	}
-
-	AddDynamicBindingMenu(MenuBuilder, Spawnable->DynamicBinding);
-}
-
-void FObjectBindingModel::AddDynamicPossessionMenu(FMenuBuilder& MenuBuilder)
-{
-	UMovieScene* MovieScene = OwnerModel->GetMovieScene();
-	FMovieScenePossessable* Possessable = MovieScene->FindPossessable(ObjectBindingID);
-	if (!Possessable)
-	{
-		return;
-	}
-
-	AddDynamicBindingMenu(MenuBuilder, Possessable->DynamicBinding);
-}
-
+	
 void FObjectBindingModel::AddDynamicBindingMenu(FMenuBuilder& MenuBuilder, FMovieSceneDynamicBinding& DynamicBinding)
 {
 	FDetailsViewArgs DetailsViewArgs;
diff --git a/Engine/Source/Editor/Sequencer/Private/MVVM/ViewModels/SequencerEditorViewModel.cpp b/Engine/Source/Editor/Sequencer/Private/MVVM/ViewModels/SequencerEditorViewModel.cpp
--- a/Engine/Source/Editor/Sequencer/Private/MVVM/ViewModels/SequencerEditorViewModel.cpp
+++ b/Engine/Source/Editor/Sequencer/Private/MVVM/ViewModels/SequencerEditorViewModel.cpp
@@ -201,6 +201,36 @@ TArrayView<const FSequencerCustomizationInfo> FSequencerEditorViewModel::GetActi
 	return ActiveCustomizationInfos;
 }
 
+TSharedPtr<FExtender> FSequencerEditorViewModel::GetSequencerMenuExtender(
+		TSharedPtr<FExtensibilityManager> ExtensibilityManager, const TArray<UObject*>& ContextObjects,
+		TFunctionRef<const FOnGetSequencerMenuExtender&(const FSequencerCustomizationInfo&)> Endpoint, FViewModelPtr InViewModel) const
+{
+	TArray<TSharedPtr<FExtender>> Extenders;
+	// Get the extenders from the extensibility manager.
+	{
+		TSharedRef<FUICommandList> CommandList(new FUICommandList);
+		TSharedPtr<FExtender> ManagerExtender = ExtensibilityManager->GetAllExtenders(CommandList, ContextObjects);
+		if (ManagerExtender)
+		{
+			Extenders.Add(ManagerExtender);
+		}
+	}
+	// Get the extenders from any active sequencer customizations.
+	for (const FSequencerCustomizationInfo& CustomizationInfo : ActiveCustomizationInfos)
+	{
+		const FOnGetSequencerMenuExtender& Delegate = ::Invoke(Endpoint, CustomizationInfo);
+		if (Delegate.IsBound())
+		{
+			TSharedPtr<FExtender> CustomizationExtender = Delegate.Execute(InViewModel);
+			if (CustomizationExtender)
+			{
+				Extenders.Add(CustomizationExtender);
+			}
+		}
+	}
+	return FExtender::Combine(Extenders);
+}
+
 } // namespace Sequencer
 } // namespace UE
 
diff --git a/Engine/Source/Editor/Sequencer/Private/SSequencer.cpp b/Engine/Source/Editor/Sequencer/Private/SSequencer.cpp
--- a/Engine/Source/Editor/Sequencer/Private/SSequencer.cpp
+++ b/Engine/Source/Editor/Sequencer/Private/SSequencer.cpp
@@ -1889,28 +1889,19 @@ bool SSequencer::IsTrackLevelFilterActive(const FString LevelName) const
 
 TSharedRef<SWidget> SSequencer::MakeActionsMenu()
 {
-	ISequencerModule& SequencerModule = FModuleManager::GetModuleChecked<ISequencerModule>("Sequencer");
+	using namespace UE::Sequencer;
+
+	TSharedPtr<FSequencer> Sequencer = SequencerPtr.Pin();
+	TSharedPtr<FSequencerEditorViewModel> EditorViewModel = Sequencer->GetViewModel();
 
+	ISequencerModule& SequencerModule = FModuleManager::GetModuleChecked<ISequencerModule>("Sequencer");
 	TSharedPtr<FExtender> ActionsMenuExtender = SequencerModule.GetActionsMenuExtensibilityManager()->GetAllExtenders();
 
-	FMenuBuilder MenuBuilder(true, SequencerPtr.Pin()->GetCommandBindings(), ActionsMenuExtender);
-	TSharedPtr<FSequencer> Sequencer = SequencerPtr.Pin();
+	FMenuBuilder MenuBuilder(true, Sequencer->GetCommandBindings(), ActionsMenuExtender);
 
 	MenuBuilder.BeginSection("SequenceOptions", LOCTEXT("SequenceOptionsHeader", "Sequence"));
 	{
-		if (SequencerPtr.Pin()->IsLevelEditorSequencer())
-		{
-			MenuBuilder.AddMenuEntry(
-				LOCTEXT("SaveAs", "Save As..."),
-				LOCTEXT("SaveAsTooltip", "Saves the current sequence under a different name"),
-				FSlateIcon(FAppStyle::GetAppStyleSetName(), "Sequencer.SaveAs"),
-				FUIAction(FExecuteAction::CreateSP(this, &SSequencer::OnSaveMovieSceneAsClicked)));
-
-			MenuBuilder.AddMenuEntry(FSequencerCommands::Get().ImportFBX);
-			MenuBuilder.AddMenuEntry(FSequencerCommands::Get().ExportFBX);
-		}
-
-		UMovieSceneSequence* RootSequence = SequencerPtr.Pin()->GetRootMovieSceneSequence();
+		UMovieSceneSequence* RootSequence = Sequencer->GetRootMovieSceneSequence();
 		if (RootSequence->GetTypedOuter<UBlueprint>() == nullptr)
 		{
 			// Only show this button where it makes sense (ie, if the sequence is not contained within a blueprint already)
@@ -3377,11 +3368,6 @@ void SSequencer::OnSaveMovieSceneClicked()
 }
 
 
-void SSequencer::OnSaveMovieSceneAsClicked()
-{
-	SequencerPtr.Pin()->SaveCurrentMovieSceneAs();
-}
-
 void SSequencer::StepToNextKey()
 {
 	StepToKey(true, false);
diff --git a/Engine/Source/Editor/Sequencer/Private/SSequencer.h b/Engine/Source/Editor/Sequencer/Private/SSequencer.h
--- a/Engine/Source/Editor/Sequencer/Private/SSequencer.h
+++ b/Engine/Source/Editor/Sequencer/Private/SSequencer.h
@@ -331,9 +331,6 @@ public:
 	/** Called when the save button is clicked */
 	void OnSaveMovieSceneClicked();
 
-	/** Called when the save-as button is clicked */
-	void OnSaveMovieSceneAsClicked();
-
 	/** Called when the curve editor is shown or hidden */
 	void OnCurveEditorVisibilityChanged(bool bShouldBeVisible);
 
diff --git a/Engine/Source/Editor/Sequencer/Private/Sequencer.cpp b/Engine/Source/Editor/Sequencer/Private/Sequencer.cpp
--- a/Engine/Source/Editor/Sequencer/Private/Sequencer.cpp
+++ b/Engine/Source/Editor/Sequencer/Private/Sequencer.cpp
@@ -55,7 +55,6 @@
 #include "Widgets/Layout/SBox.h"
 #include "Widgets/Input/SButton.h"
 #include "Styling/AppStyle.h"
-#include "Exporters/Exporter.h"
 #include "Editor/UnrealEdEngine.h"
 #include "Camera/CameraActor.h"
 #include "Engine/Selection.h"
@@ -63,7 +62,6 @@
 #include "LevelEditorViewport.h"
 #include "EditorModeManager.h"
 #include "UnrealEdMisc.h"
-#include "EditorDirectories.h"
 #include "FileHelpers.h"
 #include "UnrealEdGlobals.h"
 #include "SequencerCommands.h"
@@ -120,9 +118,7 @@
 #include "CameraRig_Crane.h"
 #include "DesktopPlatformModule.h"
 #include "Factories.h"
-#include "FbxExporter.h"
 #include "ObjectBindingTagCache.h"
-#include "UnrealExporter.h"
 #include "ISequencerEditorObjectBinding.h"
 #include "LevelSequence.h"
 #include "LevelSequenceActor.h"
@@ -139,7 +135,6 @@
 #include "MovieSceneTimeHelpers.h"
 #include "FrameNumberNumericInterface.h"
 #include "UObject/StrongObjectPtr.h"
-#include "SequencerExportTask.h"
 #include "LevelUtils.h"
 #include "Engine/Blueprint.h"
 #include "MovieSceneSequenceEditor.h"
@@ -8429,34 +8424,6 @@ void FSequencer::SaveSelectedNodesSpawnableState()
 	NotifyMovieSceneDataChanged(EMovieSceneDataChangeType::MovieSceneStructureItemsChanged);
 }
 
-void FSequencer::SetSelectedNodesSpawnableLevel(FName InLevelName)
-{
-	using namespace UE::Sequencer;
-
-	UMovieScene* MovieScene = GetFocusedMovieSceneSequence()->GetMovieScene();
-
-	if (MovieScene->IsReadOnly())
-	{
-		FSequencerUtilities::ShowReadOnlyError();
-		return;
-	}
-
-	const FScopedTransaction Transaction( LOCTEXT("SetSpawnableLevel", "Set Spawnable Level") );
-
-	MovieScene->Modify();
-
-	TArray<FMovieSceneSpawnable*> Spawnables;
-
-	for (TViewModelPtr<IObjectBindingExtension> ObjectBindingNode : ViewModel->GetSelection()->Outliner.Filter<IObjectBindingExtension>())
-	{
-		FMovieSceneSpawnable* Spawnable = MovieScene->FindSpawnable(ObjectBindingNode->GetObjectGuid());
-		if (Spawnable)
-		{
-			Spawnable->SetLevelName(InLevelName);
-		}
-	}
-}
-
 void FSequencer::ConvertToSpawnable(TSharedRef<UE::Sequencer::FObjectBindingModel> NodeToBeConverted)
 {
 	using namespace UE::Sequencer;
@@ -10217,217 +10184,6 @@ void FSequencer::RebindPossessableReferences()
 	}
 }
 
-void FSequencer::ImportFBX()
-{
-	using namespace UE::Sequencer;
-
-	TMap<FGuid, FString> ObjectBindingNameMap;
-
-	TParentFirstChildIterator<IObjectBindingExtension> ObjectBindingIt = NodeTree->GetRootNode()->GetDescendantsOfType<IObjectBindingExtension>();
-
-	// Only visit the first object binding in a given sub-hierarchy so we get top-level object bindings.
-	// This is done by skipping the branch on each iteration
-	for ( ; ObjectBindingIt; ++ObjectBindingIt)
-	{
-		FGuid ObjectBinding = ObjectBindingIt->GetObjectGuid();
-		ObjectBindingNameMap.Add(ObjectBinding, (*ObjectBindingIt).AsModel()->CastThisChecked<IOutlinerExtension>()->GetLabel().ToString());
-
-		ObjectBindingIt.IgnoreCurrentChildren();
-	}
-
-	MovieSceneToolHelpers::ImportFBXWithDialog(GetFocusedMovieSceneSequence(), *this, ObjectBindingNameMap, TOptional<bool>());
-}
-
-void FSequencer::ImportFBXOntoSelectedNodes()
-{
-	using namespace UE::Sequencer;
-
-	// The object binding and names to match when importing from fbx
-	TMap<FGuid, FString> ObjectBindingNameMap;
-
-	for (TViewModelPtr<FObjectBindingModel> ObjectBindingNode : ViewModel->GetSelection()->Outliner.Filter<FObjectBindingModel>())
-	{
-		FGuid ObjectBinding = ObjectBindingNode->GetObjectGuid();
-		ObjectBindingNameMap.Add(ObjectBinding, ObjectBindingNode->GetLabel().ToString());
-	}
-
-	MovieSceneToolHelpers::ImportFBXWithDialog(GetFocusedMovieSceneSequence(), *this, ObjectBindingNameMap, TOptional<bool>(false));
-}
-
-void FSequencer::ExportFBX()
-{
-	using namespace UE::Sequencer;
-
-	TArray<UExporter*> Exporters;
-	TArray<FString> SaveFilenames;
-	IDesktopPlatform* DesktopPlatform = FDesktopPlatformModule::Get();
-	bool bExportFileNamePicked = false;
-	if ( DesktopPlatform != NULL )
-	{
-		FString FileTypes = "FBX document|*.fbx";
-		UMovieSceneSequence* Sequence = GetFocusedMovieSceneSequence();
-		for (TObjectIterator<UClass> It; It; ++It)
-		{
-			if (!It->IsChildOf(UExporter::StaticClass()) || It->HasAnyClassFlags(CLASS_Abstract | CLASS_Deprecated | CLASS_NewerVersionExists))
-			{
-				continue;
-			}
-
-			UExporter* Default = It->GetDefaultObject<UExporter>();
-			if (!Default->SupportsObject(Sequence))
-			{
-				continue;
-			}
-
-			for (int32 i = 0; i < Default->FormatExtension.Num(); ++i)
-			{
-				const FString& FormatExtension = Default->FormatExtension[i];
-				const FString& FormatDescription = Default->FormatDescription[i];
-
-				if (FileTypes.Len() > 0)
-				{
-					FileTypes += TEXT("|");
-				}
-				FileTypes += FormatDescription;
-				FileTypes += TEXT("|*.");
-				FileTypes += FormatExtension;
-			}
-
-			Exporters.Add(Default);
-		}
-
-		bExportFileNamePicked = DesktopPlatform->SaveFileDialog(
-			FSlateApplication::Get().FindBestParentWindowHandleForDialogs(nullptr),
-			LOCTEXT( "ExportLevelSequence", "Export Level Sequence" ).ToString(),
-			*( FEditorDirectories::Get().GetLastDirectory( ELastDirectory::FBX ) ),
-			TEXT( "" ),
-			*FileTypes,
-			EFileDialogFlags::None,
-			SaveFilenames );
-	}
-
-	if ( bExportFileNamePicked )
-	{
-		FString ExportFilename = SaveFilenames[0];
-		FEditorDirectories::Get().SetLastDirectory( ELastDirectory::FBX, FPaths::GetPath( ExportFilename ) ); // Save path as default for next time.
-
-		// Make sure external selection is up to date since export could happen on tracks that have been right clicked but not have their underlying bound objects selected yet since that happens on mouse up.
-		SynchronizeExternalSelectionWithSequencerSelection();
-		
-		// Select selected nodes if there are selected nodes
-		TArray<FGuid> Bindings;
-		TArray<UMovieSceneTrack*> Tracks;
-		UMovieScene* MovieScene = GetFocusedMovieSceneSequence()->GetMovieScene();
-		for (FViewModelPtr Node : ViewModel->GetSelection()->Outliner)
-		{
-			if (TViewModelPtr<IObjectBindingExtension> ObjectBindingNode = Node.ImplicitCast())
-			{
-				Bindings.Add(ObjectBindingNode->GetObjectGuid());
-
-				for (TViewModelPtr<IOutlinerExtension> DescendantNode : Node->GetDescendantsOfType<IOutlinerExtension>())
-				{
-					if (!ViewModel->GetSelection()->Outliner.IsSelected(DescendantNode) && DescendantNode.AsModel()->IsA<IObjectBindingExtension>())
-					{
-						IObjectBindingExtension* DescendantObjectBindingNode = DescendantNode.AsModel()->CastThisChecked<IObjectBindingExtension>();
-						Bindings.Add(DescendantObjectBindingNode->GetObjectGuid());
-					}
-				}
-			}
-			else if (TViewModelPtr<ITrackExtension> TrackNode = Node.ImplicitCast())
-			{
-				UMovieSceneTrack* Track = TrackNode->GetTrack();
-				if (Track && MovieScene->ContainsTrack(*Track))
-				{
-					Tracks.Add(Track);
-				}
-			}
-		}
-
-		FString FileExtension = FPaths::GetExtension(ExportFilename);
-		if (FileExtension == TEXT("fbx"))
-		{
-			ExportFBXInternal(ExportFilename, Bindings, (Bindings.Num() + Tracks.Num()) > 0 ? Tracks : MovieScene->GetTracks());
-		}
-		else
-		{
-			for (UExporter* Exporter : Exporters)
-			{
-				if (Exporter->FormatExtension.Contains(FileExtension))
-				{
-					USequencerExportTask* ExportTask = NewObject<USequencerExportTask>();
-					TStrongObjectPtr<USequencerExportTask> ExportTaskGuard(ExportTask);
-					ExportTask->Object = GetFocusedMovieSceneSequence();
-					ExportTask->Exporter = nullptr;
-					ExportTask->Filename = ExportFilename;
-					ExportTask->bSelected = false;
-					ExportTask->bReplaceIdentical = true;
-					ExportTask->bPrompt = false;
-					ExportTask->bUseFileArchive = false;
-					ExportTask->bWriteEmptyFiles = false;
-					ExportTask->bAutomated = false;
-					ExportTask->Exporter = NewObject<UExporter>(GetTransientPackage(), Exporter->GetClass());
-
-					ExportTask->SequencerContext = GetPlaybackContext();
-
-					UExporter::RunAssetExportTask(ExportTask);
-
-					ExportTask->Object = nullptr;
-					ExportTask->Exporter = nullptr;
-					ExportTask->SequencerContext = nullptr;
-
-					break;
-				}
-			}
-		}
-	}
-}
-
-
-void FSequencer::ExportFBXInternal(const FString& ExportFilename, const TArray<FGuid>& Bindings, const TArray<UMovieSceneTrack*>& Tracks)
-{
-	{
-		UnFbx::FFbxExporter* Exporter = UnFbx::FFbxExporter::GetInstance();
-		//Show the fbx export dialog options
-		bool ExportCancel = false;
-		bool ExportAll = false;
-		Exporter->FillExportOptions(false, true, ExportFilename, ExportCancel, ExportAll);
-		if (!ExportCancel)
-		{
-			UMovieScene* MovieScene = GetFocusedMovieSceneSequence()->GetMovieScene();
-			UWorld* World = Cast<UWorld>(GetPlaybackContext());
-			FMovieSceneSequenceIDRef Template = GetFocusedTemplateID();
-			UnFbx::FFbxExporter::FLevelSequenceNodeNameAdapter NodeNameAdapter(MovieScene, this, Template);
-
-			{
-				FSpawnableRestoreState SpawnableRestoreState(MovieScene);
-				if (SpawnableRestoreState.bWasChanged)
-				{
-					// Evaluate at the beginning of the subscene time to ensure that spawnables are created before export
-					SetLocalTimeDirectly(UE::MovieScene::DiscreteInclusiveLower(GetTimeBounds()));
-				}
-
-				if (MovieSceneToolHelpers::ExportFBX(World, MovieScene, this, Bindings, Tracks, NodeNameAdapter, Template, ExportFilename, RootToLocalTransform))
-				{
-					FNotificationInfo Info(NSLOCTEXT("Sequencer", "ExportFBXSucceeded", "FBX Export Succeeded."));
-					Info.Hyperlink = FSimpleDelegate::CreateStatic([](FString InFilename) { FPlatformProcess::ExploreFolder(*InFilename); }, ExportFilename);
-					Info.HyperlinkText = FText::FromString(ExportFilename);
-					Info.ExpireDuration = 5.0f;
-					FSlateNotificationManager::Get().AddNotification(Info)->SetCompletionState(SNotificationItem::CS_Success);
-				}
-				else
-				{
-					FNotificationInfo Info(NSLOCTEXT("Sequencer", "ExportFBXFailed", "FBX Export Failed."));
-					Info.ExpireDuration = 5.0f;
-					FSlateNotificationManager::Get().AddNotification(Info)->SetCompletionState(SNotificationItem::CS_Fail);
-				}
-			}
-
-			ForceEvaluate();
-		}
-	}
-}
-
-
 void FSequencer::GenericTextEntryModeless(const FText& DialogText, const FText& DefaultText, FOnTextCommitted OnTextComitted)
 {
 	TSharedRef<STextEntryPopup> TextEntryPopup = 
@@ -11205,16 +10961,6 @@ void FSequencer::BindCommands()
 		FExecuteAction::CreateSP( this, &FSequencer::RebindPossessableReferences ),
 		FCanExecuteAction::CreateLambda( []{ return true; } ) );
 
-	SequencerCommandBindings->MapAction(
-		Commands.ImportFBX,
-		FExecuteAction::CreateSP( this, &FSequencer::ImportFBX ),
-		FCanExecuteAction::CreateLambda( [] { return true; } ) );
-
-	SequencerCommandBindings->MapAction(
-		Commands.ExportFBX,
-		FExecuteAction::CreateSP( this, &FSequencer::ExportFBX ),
-		FCanExecuteAction::CreateLambda( [] { return true; } ) );
-
 	SequencerCommandBindings->MapAction(
 		Commands.MoveToNewFolder,
 		FExecuteAction::CreateSP( this, &FSequencer::MoveSelectedNodesToNewFolder ),
diff --git a/Engine/Source/Editor/Sequencer/Private/Sequencer.h b/Engine/Source/Editor/Sequencer/Private/Sequencer.h
--- a/Engine/Source/Editor/Sequencer/Private/Sequencer.h
+++ b/Engine/Source/Editor/Sequencer/Private/Sequencer.h
@@ -357,11 +357,6 @@ public:
 	 */
 	void ConvertSelectedNodesToPossessables();
 
-	/*
-	 * Set the spawnable level for the currently selected objects
-	 */
-	void SetSelectedNodesSpawnableLevel(FName InLevelName);
-
 protected:
 
 	/**
@@ -641,13 +636,6 @@ public:
 	/** Rebinds all possessable references in the current sequence to update them to the latest referencing mechanism. */
 	void RebindPossessableReferences();
 
-	/** Imports the animation from an fbx file. */
-	void ImportFBX();
-	void ImportFBXOntoSelectedNodes();
-
-	/** Exports the animation to an fbx file. */
-	void ExportFBX();
-
 public:
 
 	/** Put the sequencer in a horizontally auto-scrolling state with the given rate */
@@ -828,6 +816,13 @@ public:
 	UE::Sequencer::FSequencerSelection& GetSelection();
 	FSequencerSelectionPreview& GetSelectionPreview();
 
+	/**
+	 * Gets the far time boundaries of the currently edited movie scene
+	 * If the scene has shots, it only takes the shot section boundaries
+	 * Otherwise, it finds the furthest boundaries of all sections
+	 */
+	TRange<FFrameNumber> GetTimeBounds() const;
+
 protected:
 
 	/** Reevaluate the sequence at the current time */
@@ -841,13 +836,6 @@ protected:
 
 	/** Update the time bounds to the focused movie scene */
 	void UpdateTimeBoundsToFocusedMovieScene();
-
-	/**
-	 * Gets the far time boundaries of the currently edited movie scene
-	 * If the scene has shots, it only takes the shot section boundaries
-	 * Otherwise, it finds the furthest boundaries of all sections
-	 */
-	TRange<FFrameNumber> GetTimeBounds() const;
 	
 	/**
 	 * Gets the time boundaries of the currently filtering shot sections.
@@ -920,9 +908,6 @@ protected:
 	/** Called when the scrub position parent sequence is changed */
 	void OnScrubPositionParentChanged(FMovieSceneSequenceID InScrubPositionParent);
 
-	/** Exports sequence to a FBX file */
-	void ExportFBXInternal(const FString& Filename, const TArray<FGuid>& Bindings, const TArray<UMovieSceneTrack*>& Tracks);
-
 protected:
 
 	/**
diff --git a/Engine/Source/Editor/Sequencer/Private/SequencerCommands.cpp b/Engine/Source/Editor/Sequencer/Private/SequencerCommands.cpp
--- a/Engine/Source/Editor/Sequencer/Private/SequencerCommands.cpp
+++ b/Engine/Source/Editor/Sequencer/Private/SequencerCommands.cpp
@@ -172,9 +172,6 @@ void FSequencerCommands::RegisterCommands()
 	UI_COMMAND( FixPossessableObjectClass, "Fix Possessable Object Class", "Try to automatically fix up possessables whose object class don't match the object class of their currently bound objects.", EUserInterfaceActionType::Button, FInputChord() );
 	UI_COMMAND( RebindPossessableReferences, "Rebind Possesable References", "Rebinds all possessables in the current sequence to ensure they're using the most robust referencing mechanism.", EUserInterfaceActionType::Button, FInputChord() );
 
-	UI_COMMAND( ImportFBX, "Import...", "Import the animation from an FBX file.", EUserInterfaceActionType::Button, FInputChord() );
-	UI_COMMAND( ExportFBX, "Export...", "Export the selected objects (all if none selected) and animation to an FBX file. (Shots and sub-scenes not supported)", EUserInterfaceActionType::Button, FInputChord() );
-
 	UI_COMMAND( ToggleEvaluateSubSequencesInIsolation, "Evaluate Sub Sequences In Isolation", "When enabled, will only evaluate the currently focused sequence; otherwise evaluate from the root sequence.", EUserInterfaceActionType::ToggleButton, FInputChord() );
 
 	UI_COMMAND( QuickTreeSearch, "Quick Tree Search", "Jumps keyboard focus to the tree searchbox to allow searching for tracks in the current Sequence.", EUserInterfaceActionType::Button, FInputChord(EModifierKey::Control, EKeys::F));
diff --git a/Engine/Source/Editor/Sequencer/Private/SequencerUtilities.cpp b/Engine/Source/Editor/Sequencer/Private/SequencerUtilities.cpp
--- a/Engine/Source/Editor/Sequencer/Private/SequencerUtilities.cpp
+++ b/Engine/Source/Editor/Sequencer/Private/SequencerUtilities.cpp
@@ -2980,4 +2980,19 @@ void FSequencerUtilities::ShowSpawnableNotAllowedError()
 	FSlateNotificationManager::Get().AddNotification(Info)->SetCompletionState(SNotificationItem::CS_Fail);	
 }
 
+void FSequencerUtilities::SaveCurrentMovieSceneAs(TSharedRef<ISequencer> Sequencer)
+{
+	StaticCastSharedPtr<FSequencer>(Sequencer.ToSharedPtr())->SaveCurrentMovieSceneAs();
+}
+
+void FSequencerUtilities::SynchronizeExternalSelectionWithSequencerSelection (TSharedRef<ISequencer> Sequencer)
+{
+	StaticCastSharedPtr<FSequencer>(Sequencer.ToSharedPtr())->SynchronizeExternalSelectionWithSequencerSelection();
+}
+
+TRange<FFrameNumber> FSequencerUtilities::GetTimeBounds(TSharedRef<ISequencer> Sequencer)
+{
+	return StaticCastSharedPtr<FSequencer>(Sequencer.ToSharedPtr())->GetTimeBounds();
+}
+
 #undef LOCTEXT_NAMESPACE
diff --git a/Engine/Source/Editor/Sequencer/Public/ISequencerModule.h b/Engine/Source/Editor/Sequencer/Public/ISequencerModule.h
--- a/Engine/Source/Editor/Sequencer/Public/ISequencerModule.h
+++ b/Engine/Source/Editor/Sequencer/Public/ISequencerModule.h
@@ -296,14 +296,14 @@ public:
 	virtual bool CanAnimateProperty(FProperty* Property) = 0;
 
 	/**
-	* Get the extensibility manager for menus.
+	* Get the extensibility manager for object binding menus.
 	*
 	* @return ObjectBinding Context Menu extensibility manager.
 	*/
 	virtual TSharedPtr<FExtensibilityManager> GetObjectBindingContextMenuExtensibilityManager() const = 0;
 
 	/**
-	 * Get the extensibility manager for menus.
+	 * Get the extensibility manager for add track menus.
 	 *
 	 * @return Add Track Menu extensibility manager.
 	 */
diff --git a/Engine/Source/Editor/Sequencer/Public/MVVM/ViewModels/ObjectBindingModel.h b/Engine/Source/Editor/Sequencer/Public/MVVM/ViewModels/ObjectBindingModel.h
--- a/Engine/Source/Editor/Sequencer/Public/MVVM/ViewModels/ObjectBindingModel.h
+++ b/Engine/Source/Editor/Sequencer/Public/MVVM/ViewModels/ObjectBindingModel.h
@@ -39,7 +39,7 @@ enum class EObjectBindingType
 	Possessable, Spawnable, Unknown
 };
 
-class FObjectBindingModel
+class SEQUENCER_API FObjectBindingModel
 	: public FOutlinerItemModel
 	, public IObjectBindingExtension
 	, public IDraggableOutlinerExtension
@@ -119,33 +119,27 @@ public:
 
 	TSharedRef<SWidget> GetAddTrackMenuContent();
 
+	/** Build a sub-menu for editing the given dynamic binding */
+	void AddDynamicBindingMenu(FMenuBuilder& MenuBuilder, FMovieSceneDynamicBinding& DynamicBinding);
+
 protected:
 
 	/*~ FViewModel interface */
 	void OnConstruct() override;
 
-private:
-
-	void AddSpawnOwnershipMenu(FMenuBuilder& MenuBuilder);
-	void AddSpawnLevelMenu(FMenuBuilder& MenuBuilder);
-	void AddTagMenu(FMenuBuilder& MenuBuilder);
-	void AddChangeClassMenu(FMenuBuilder& MenuBuilder);
-	void AddDynamicSpawnMenu(FMenuBuilder& MenuBuilder);
-	void AddDynamicPossessionMenu(FMenuBuilder& MenuBuilder);
-
 private:
 
 	void AddPropertyMenuItems(FMenuBuilder& AddTrackMenuBuilder, TArray<FPropertyPath> KeyableProperties, int32 PropertyNameIndexStart, int32 PropertyNameIndexEnd);
 	void HandleAddTrackSubMenuNew(FMenuBuilder& AddTrackMenuBuilder, TArray<FPropertyPath> KeyablePropertyPaths, int32 PropertyNameIndexStart);
 	void HandlePropertyMenuItemExecute(FPropertyPath PropertyPath);
 
+	void AddTagMenu(FMenuBuilder& MenuBuilder);
 	ECheckBoxState GetTagCheckState(FName TagName);
 	void ToggleTag(FName TagName);
 	void HandleDeleteTag(FName TagName);
 	void HandleAddTag(FName TagName);
 	void HandleTemplateActorClassPicked(UClass* ChosenClass);
 
-	void AddDynamicBindingMenu(FMenuBuilder& MenuBuilder, FMovieSceneDynamicBinding& DynamicBinding);
 	void OnFinishedChangingDynamicBindingProperties(const FPropertyChangedEvent& ChangeEvent, TSharedPtr<FStructOnScope> ValueStruct);
 
 protected:
diff --git a/Engine/Source/Editor/Sequencer/Public/MVVM/ViewModels/SequencerEditorViewModel.h b/Engine/Source/Editor/Sequencer/Public/MVVM/ViewModels/SequencerEditorViewModel.h
--- a/Engine/Source/Editor/Sequencer/Public/MVVM/ViewModels/SequencerEditorViewModel.h
+++ b/Engine/Source/Editor/Sequencer/Public/MVVM/ViewModels/SequencerEditorViewModel.h
@@ -54,6 +54,11 @@ public:
 	/** Get the active customization infos */
 	TArrayView<const FSequencerCustomizationInfo> GetActiveCustomizationInfos() const;
 
+	/** Build a combined menu extender */
+	TSharedPtr<FExtender> GetSequencerMenuExtender(
+			TSharedPtr<FExtensibilityManager> ExtensibilityManager, const TArray<UObject*>& ContextObjects,
+			TFunctionRef<const FOnGetSequencerMenuExtender&(const FSequencerCustomizationInfo&)> Endpoint, FViewModelPtr InViewModel) const;
+
 	/** Gets the pinned track area view-model. */
 	TSharedPtr<FTrackAreaViewModel> GetPinnedTrackArea() const;
 
diff --git a/Engine/Source/Editor/Sequencer/Public/SequencerCommands.h b/Engine/Source/Editor/Sequencer/Public/SequencerCommands.h
--- a/Engine/Source/Editor/Sequencer/Public/SequencerCommands.h
+++ b/Engine/Source/Editor/Sequencer/Public/SequencerCommands.h
@@ -12,7 +12,7 @@
 
 class FUICommandInfo;
 
-class FSequencerCommands : public TCommands<FSequencerCommands>
+class SEQUENCER_API FSequencerCommands : public TCommands<FSequencerCommands>
 {
 
 public:
@@ -430,12 +430,6 @@ public:
 	/** Rebinds all possessable references with their current bindings. */
 	TSharedPtr< FUICommandInfo > RebindPossessableReferences;
 
-	/** Imports animation from fbx. */
-	TSharedPtr< FUICommandInfo > ImportFBX;
-
-	/** Exports animation to fbx. */
-	TSharedPtr< FUICommandInfo > ExportFBX;
-
 	/** Toggle whether we should evaluate sub sequences in isolation */
 	TSharedPtr< FUICommandInfo > ToggleEvaluateSubSequencesInIsolation;
 
diff --git a/Engine/Source/Editor/Sequencer/Public/SequencerCustomizationManager.h b/Engine/Source/Editor/Sequencer/Public/SequencerCustomizationManager.h
--- a/Engine/Source/Editor/Sequencer/Public/SequencerCustomizationManager.h
+++ b/Engine/Source/Editor/Sequencer/Public/SequencerCustomizationManager.h
@@ -5,6 +5,7 @@
 #include "CoreMinimal.h"
 #include "ISequencerModule.h"
 #include "Misc/EnumClassFlags.h"
+#include "MVVM/ViewModelPtr.h"
 
 class FAssetDragDropOp;
 class FClassDragDropOp;
@@ -49,6 +50,9 @@ ENUM_CLASS_FLAGS(ESequencerPasteSupport);
 /** A delegate that customizes how the sequencer behaves when data is pasted into it */
 DECLARE_DELEGATE_RetVal(ESequencerPasteSupport, FOnSequencerPaste);
 
+/** A delegate that returns an extender */
+DECLARE_DELEGATE_RetVal_OneParam(TSharedPtr<FExtender>, FOnGetSequencerMenuExtender, UE::Sequencer::FViewModelPtr);
+
 /** Class for specifying customizations to apply to a sequence editor. */
 struct FSequencerCustomizationInfo
 {
@@ -58,6 +62,9 @@ struct FSequencerCustomizationInfo
 	/** Extender for the sequencer toolbar. */
 	TSharedPtr<FExtender> ToolbarExtender;
 
+	/** Extender for the object binding. */
+	FOnGetSequencerMenuExtender OnBuildObjectBindingContextMenu;
+
 	/** Called when something is dragged over the sequencer. */
 	FOptionalOnDragDrop OnReceivedDragOver;
 
diff --git a/Engine/Source/Editor/Sequencer/Public/SequencerExportTask.h b/Engine/Source/Editor/Sequencer/Public/SequencerExportTask.h
deleted file mode 100644
diff --git a/Engine/Source/Editor/Sequencer/Public/SequencerUtilities.h b/Engine/Source/Editor/Sequencer/Public/SequencerUtilities.h
--- a/Engine/Source/Editor/Sequencer/Public/SequencerUtilities.h
+++ b/Engine/Source/Editor/Sequencer/Public/SequencerUtilities.h
@@ -194,4 +194,10 @@ struct SEQUENCER_API FSequencerUtilities
 	
 	/** Show an error if spawnable is not allowed in a movie scene*/
 	static void ShowSpawnableNotAllowedError();
+
+	// Methods exposing FSequencer functionality outside of the Sequencer module.
+	// This is needed while things are moving into view-models here and there, step by step.
+	static void SaveCurrentMovieSceneAs(TSharedRef<ISequencer> Sequencer);
+	static void SynchronizeExternalSelectionWithSequencerSelection (TSharedRef<ISequencer> Sequencer);
+	static TRange<FFrameNumber> GetTimeBounds(TSharedRef<ISequencer> Sequencer);
 };
diff --git a/Engine/Source/Editor/Sequencer/Sequencer.Build.cs b/Engine/Source/Editor/Sequencer/Sequencer.Build.cs
--- a/Engine/Source/Editor/Sequencer/Sequencer.Build.cs
+++ b/Engine/Source/Editor/Sequencer/Sequencer.Build.cs
@@ -6,12 +6,6 @@ public class Sequencer : ModuleRules
 {
 	public Sequencer(ReadOnlyTargetRules Target) : base(Target)
 	{
-		PrivateIncludePaths.AddRange(
-			new string[] {
-				"Editor/UnrealEd/Private" // TODO: Fix this, for now it's needed for the fbx exporter
-				}
-			);
-
 		PublicDependencyModuleNames.AddRange(
 			new string[] {
 				"TimeManagement",
@@ -87,6 +81,5 @@ public class Sequencer : ModuleRules
 			);
 
 		CircularlyReferencedDependentModules.Add("MovieSceneTools");
-		AddEngineThirdPartyPrivateStaticDependencies(Target, "FBX");
 	}
 }
