[LevelSequenceEditor] Expose playback rate of LevelSequenceEditor to Blueprint

diff --git a/Engine/Plugins/MovieScene/LevelSequenceEditor/Source/LevelSequenceEditor/Private/LevelSequenceEditorBlueprintLibrary.cpp b/Engine/Plugins/MovieScene/LevelSequenceEditor/Source/LevelSequenceEditor/Private/LevelSequenceEditorBlueprintLibrary.cpp
--- a/Engine/Plugins/MovieScene/LevelSequenceEditor/Source/LevelSequenceEditor/Private/LevelSequenceEditorBlueprintLibrary.cpp
+++ b/Engine/Plugins/MovieScene/LevelSequenceEditor/Source/LevelSequenceEditor/Private/LevelSequenceEditorBlueprintLibrary.cpp
@@ -165,6 +165,24 @@ int32 ULevelSequenceEditorBlueprintLibrary::GetCurrentLocalTime()
 	return 0;
 }
 
+void ULevelSequenceEditorBlueprintLibrary::SetPlaybackSpeed(float NewPlaybackSpeed)
+{
+	if (CurrentSequencer.IsValid())
+	{
+		CurrentSequencer.Pin()->SetPlaybackSpeed(NewPlaybackSpeed);
+	}
+}
+
+float ULevelSequenceEditorBlueprintLibrary::GetPlaybackSpeed()
+{
+	if (CurrentSequencer.IsValid())
+	{
+		return CurrentSequencer.Pin()->GetPlaybackSpeed();
+	}
+
+	return 0.f;
+}
+
 void ULevelSequenceEditorBlueprintLibrary::PlayTo(FMovieSceneSequencePlaybackParams PlaybackParams)
 {
 	if (CurrentSequencer.IsValid())
diff --git a/Engine/Plugins/MovieScene/LevelSequenceEditor/Source/LevelSequenceEditor/Public/LevelSequenceEditorBlueprintLibrary.h b/Engine/Plugins/MovieScene/LevelSequenceEditor/Source/LevelSequenceEditor/Public/LevelSequenceEditorBlueprintLibrary.h
--- a/Engine/Plugins/MovieScene/LevelSequenceEditor/Source/LevelSequenceEditor/Public/LevelSequenceEditorBlueprintLibrary.h
+++ b/Engine/Plugins/MovieScene/LevelSequenceEditor/Source/LevelSequenceEditor/Public/LevelSequenceEditorBlueprintLibrary.h
@@ -125,6 +125,18 @@ public:
 	UFUNCTION(BlueprintCallable, Category = "Level Sequence Editor")
 	static int32 GetCurrentLocalTime();
 
+	/**
+	 * Set playback speed of the current level sequence
+	 */
+	UFUNCTION(BlueprintCallable, Category = "Level Sequence Editor")
+	static void SetPlaybackSpeed(float NewPlaybackSpeed);
+
+	/**
+	 * Get playback speed of the current level sequence
+	 */
+	UFUNCTION(BlueprintCallable, Category = "Level Sequence Editor")
+	static float GetPlaybackSpeed();
+
 	/**
 	 * Play from the current time to the requested time in frames
 	 */
