Remove UE_DEPRECATED(5.1)

#rnx

diff --git a/Engine/Plugins/Cameras/GameplayCameras/Source/GameplayCameras/Public/LegacyCameraShake.h b/Engine/Plugins/Cameras/GameplayCameras/Source/GameplayCameras/Public/LegacyCameraShake.h
--- a/Engine/Plugins/Cameras/GameplayCameras/Source/GameplayCameras/Public/LegacyCameraShake.h
+++ b/Engine/Plugins/Cameras/GameplayCameras/Source/GameplayCameras/Public/LegacyCameraShake.h
@@ -312,13 +312,6 @@ private:
 	virtual void TeardownShakePatternImpl() override;
 };
 
-/** Backwards compatible name for the legacy camera shake, for C++ code. */
-UE_DEPRECATED(4.26, "Please use ULegacyCameraShake")
-typedef ULegacyCameraShake UCameraShake;
-
-UE_DEPRECATED(5.1, "Please use ULegacyCameraShake")
-typedef ULegacyCameraShake UMatineeCameraShake;
-
 /**
  * Blueprint function library for autocasting from a base camera shake to a legacy camera shake.
  * This prevents breaking Blueprints now that APlayerCameraManager::StartCameraShake returns the base class.
diff --git a/Engine/Plugins/MovieScene/LevelSequenceEditor/Source/LevelSequenceEditor/Private/LevelSequenceEditorBlueprintLibrary.cpp b/Engine/Plugins/MovieScene/LevelSequenceEditor/Source/LevelSequenceEditor/Private/LevelSequenceEditorBlueprintLibrary.cpp
--- a/Engine/Plugins/MovieScene/LevelSequenceEditor/Source/LevelSequenceEditor/Private/LevelSequenceEditorBlueprintLibrary.cpp
+++ b/Engine/Plugins/MovieScene/LevelSequenceEditor/Source/LevelSequenceEditor/Private/LevelSequenceEditorBlueprintLibrary.cpp
@@ -251,16 +251,6 @@ TArray<UMovieSceneFolder*> ULevelSequenceEditorBlueprintLibrary::GetSelectedFold
 	return OutSelectedFolders;
 }
 
-TArray<FGuid> ULevelSequenceEditorBlueprintLibrary::GetSelectedObjects()
-{
-	TArray<FGuid> OutSelectedGuids;
-	if (CurrentSequencer.IsValid())
-	{
-		CurrentSequencer.Pin()->GetSelectedObjects(OutSelectedGuids);
-	}
-	return OutSelectedGuids;
-}
-
 TArray<FMovieSceneBindingProxy> ULevelSequenceEditorBlueprintLibrary::GetSelectedBindings()
 {
 	TArray<FMovieSceneBindingProxy> OutSelectedBindings;
@@ -329,17 +319,6 @@ void ULevelSequenceEditorBlueprintLibrary::SelectFolders(const TArray<UMovieScen
 	}
 }
 
-void ULevelSequenceEditorBlueprintLibrary::SelectObjects(TArray<FGuid> ObjectBindings)
-{
-	if (CurrentSequencer.IsValid())
-	{
-		for (FGuid ObjectBinding : ObjectBindings)
-		{
-			CurrentSequencer.Pin()->SelectObject(ObjectBinding);
-		}
-	}
-}
-
 void ULevelSequenceEditorBlueprintLibrary::SelectBindings(const TArray<FMovieSceneBindingProxy>& ObjectBindings)
 {
 	if (CurrentSequencer.IsValid())
diff --git a/Engine/Plugins/MovieScene/LevelSequenceEditor/Source/LevelSequenceEditor/Public/LevelSequenceEditorBlueprintLibrary.h b/Engine/Plugins/MovieScene/LevelSequenceEditor/Source/LevelSequenceEditor/Public/LevelSequenceEditorBlueprintLibrary.h
--- a/Engine/Plugins/MovieScene/LevelSequenceEditor/Source/LevelSequenceEditor/Public/LevelSequenceEditorBlueprintLibrary.h
+++ b/Engine/Plugins/MovieScene/LevelSequenceEditor/Source/LevelSequenceEditor/Public/LevelSequenceEditorBlueprintLibrary.h
@@ -167,11 +167,6 @@ public:
 	UFUNCTION(BlueprintPure, Category = "Level Sequence Editor")
 	static TArray<UMovieSceneFolder*> GetSelectedFolders();
 
-	/** Gets the currently selected Object Guids*/
-	UE_DEPRECATED(5.1, "GetSelectedObjects is deprecated, please use GetSelectedBindings which returns an array of FMovieSceneBindingProxy")
-	UFUNCTION(BlueprintPure, Category = "Level Sequence Editor", meta = (DeprecatedFunction, DeprecationMessage="GetSelectedObjects is deprecated, please use GetSelectedBindings which returns an array of FMovieSceneBindingProxy"))
-	static TArray<FGuid> GetSelectedObjects();
-
 	/** Gets the currently selected object bindings */
 	UFUNCTION(BlueprintPure, Category = "Level Sequence Editor")
 	static TArray<FMovieSceneBindingProxy> GetSelectedBindings();
@@ -192,11 +187,6 @@ public:
 	UFUNCTION(BlueprintCallable, Category = "Level Sequence Editor")
 	static void SelectFolders(const TArray<UMovieSceneFolder*>& Folders);
 
-	/** Select objects by GUID */
-	UE_DEPRECATED(5.1, "SelectObjects is deprecated, please use SelectBindings which takes an FMovieSceneBindingProxy")
-	UFUNCTION(BlueprintCallable, Category = "Level Sequence Editor", meta=(DeprecatedFunction, DeprecationMessage="SelectObjects is deprecated, please use SelectBindings which takes an FMovieSceneBindingProxy"))
-	static void SelectObjects(TArray<FGuid> ObjectBinding);
-
 	/** Select bindings */
 	UFUNCTION(BlueprintCallable, Category = "Level Sequence Editor")
 	static void SelectBindings(const TArray<FMovieSceneBindingProxy>& ObjectBindings);
diff --git a/Engine/Plugins/MovieScene/SequencerScripting/Source/SequencerScripting/Public/SequencerBindingProxy.h b/Engine/Plugins/MovieScene/SequencerScripting/Source/SequencerScripting/Public/SequencerBindingProxy.h
--- a/Engine/Plugins/MovieScene/SequencerScripting/Source/SequencerScripting/Public/SequencerBindingProxy.h
+++ b/Engine/Plugins/MovieScene/SequencerScripting/Source/SequencerScripting/Public/SequencerBindingProxy.h
@@ -1,15 +1,16 @@
 // Copyright Epic Games, Inc. All Rights Reserved.
 
+// HEADER_UNIT_SKIP - Deprecated
+/**
+ * DEPRECATED MovieSceneBindingProxy.h Use SequencerBindingProxy.h instead
+ */
+ 
 #pragma once
 
-#include "CoreTypes.h"
-
-struct FMovieSceneBindingProxy;
-
-UE_DEPRECATED(5.1, "Use FMovieSceneBindingProxy instead.") 
-typedef FMovieSceneBindingProxy FSequencerBindingProxy;
+#ifdef _MSC_VER
+#pragma message(__FILE__"(9): warning: use MovieSceneBindingProxy.h instead of SequencerBindingProxy.h")
+#else
+#pragma message("#include MovieSceneBindingProxy.h instead of SequencerBindingProxy.h")
+#endif
 
-#if UE_ENABLE_INCLUDE_ORDER_DEPRECATED_IN_5_2
 #include "MovieSceneBindingProxy.h"
-#include "UObject/ObjectMacros.h"
-#endif
diff --git a/Engine/Plugins/MovieScene/SequencerScripting/Source/SequencerScriptingEditor/Public/SequencerTools.h b/Engine/Plugins/MovieScene/SequencerScripting/Source/SequencerScriptingEditor/Public/SequencerTools.h
--- a/Engine/Plugins/MovieScene/SequencerScripting/Source/SequencerScriptingEditor/Public/SequencerTools.h
+++ b/Engine/Plugins/MovieScene/SequencerScripting/Source/SequencerScriptingEditor/Public/SequencerTools.h
@@ -170,9 +170,6 @@ public:
 	UFUNCTION(BlueprintCallable, Category = "Editor Scripting | Sequencer Tools | FBX")
 	static bool ExportLevelSequenceFBX(const FSequencerExportFBXParams& InParams);
 	
-	UE_DEPRECATED(5.1, "Please use ExportLevelSequenceFBX that takes a FSequencerExportFBXParams")
-	static bool ExportLevelSequenceFBX(UWorld* InWorld, ULevelSequence* InSequence, const TArray<FMovieSceneBindingProxy>& InBindings, const TArray<UMovieSceneTrack*>& InTracks, UFbxExportOption* InOverrideOptions, const FString& InFBXFileName) { FSequencerExportFBXParams Params(InWorld, InSequence, InSequence, InBindings, InTracks, InOverrideOptions, InFBXFileName); return ExportLevelSequenceFBX(Params); }
-
 	/*
 	 * Export Passed in Binding as an Anim Seqquence.
 	 *
diff --git a/Engine/Source/Editor/Sequencer/Private/SequencerSectionPainter.cpp b/Engine/Source/Editor/Sequencer/Private/SequencerSectionPainter.cpp
--- a/Engine/Source/Editor/Sequencer/Private/SequencerSectionPainter.cpp
+++ b/Engine/Source/Editor/Sequencer/Private/SequencerSectionPainter.cpp
@@ -6,8 +6,7 @@
 #include "MVVM/ViewModels/SectionModel.h"
 
 FSequencerSectionPainter::FSequencerSectionPainter(FSlateWindowElementList& OutDrawElements, const FGeometry& InSectionGeometry, TSharedPtr<UE::Sequencer::FSectionModel> InSection)
-	: Section(*InSection->GetSection())
-	, SectionModel(InSection)
+	: SectionModel(InSection)
 	, DrawElements(OutDrawElements)
 	, SectionGeometry(InSectionGeometry)
 	, LayerId(0)
diff --git a/Engine/Source/Editor/Sequencer/Public/SequencerSectionPainter.h b/Engine/Source/Editor/Sequencer/Public/SequencerSectionPainter.h
--- a/Engine/Source/Editor/Sequencer/Public/SequencerSectionPainter.h
+++ b/Engine/Source/Editor/Sequencer/Public/SequencerSectionPainter.h
@@ -67,10 +67,6 @@ public:
 
 public:
 
-	/** The movie scene section we're painting */
-	UE_DEPRECATED(5.1, "Please use SectionModel->GetSection() instead")
-	UMovieSceneSection& Section;
-
 	/** The section we are painting */
 	TSharedPtr<UE::Sequencer::FSectionModel> SectionModel;
 
diff --git a/Engine/Source/Runtime/LevelSequence/Public/LevelSequencePlayer.h b/Engine/Source/Runtime/LevelSequence/Public/LevelSequencePlayer.h
--- a/Engine/Source/Runtime/LevelSequence/Public/LevelSequencePlayer.h
+++ b/Engine/Source/Runtime/LevelSequence/Public/LevelSequencePlayer.h
@@ -98,13 +98,6 @@ public:
 	 */
 	LEVELSEQUENCE_API void Initialize(ULevelSequence* InLevelSequence, ULevel* InLevel, const FLevelSequenceCameraSettings& InCameraSettings);
 
-	UE_DEPRECATED(5.1, "Use SetPlaybackSettings(...) then Initialize(ULevelSequence*, ULevel*, const FLevelSequenceCameraSettings&)")
-	void Initialize(ULevelSequence* InLevelSequence, ULevel* InLevel, const FMovieSceneSequencePlaybackSettings& Settings, const FLevelSequenceCameraSettings& InCameraSettings)
-	{
-		SetPlaybackSettings(Settings);
-		Initialize(InLevelSequence, InLevel, InCameraSettings);
-	}
-
 public:
 
 	/**
