Sequencer: Fixed PasteBindings not returning anything.

Use OutObjectBindings instead of OutBindings

diff --git a/Engine/Plugins/MovieScene/LevelSequenceEditor/Source/LevelSequenceEditor/Private/LevelSequenceEditorSubsystem.cpp b/Engine/Plugins/MovieScene/LevelSequenceEditor/Source/LevelSequenceEditor/Private/LevelSequenceEditorSubsystem.cpp
--- a/Engine/Plugins/MovieScene/LevelSequenceEditor/Source/LevelSequenceEditor/Private/LevelSequenceEditorSubsystem.cpp
+++ b/Engine/Plugins/MovieScene/LevelSequenceEditor/Source/LevelSequenceEditor/Private/LevelSequenceEditorSubsystem.cpp
@@ -454,8 +454,7 @@ bool ULevelSequenceEditorSubsystem::PasteBindings(const FString& InTextToImport,
 	}
 
 	TArray<FNotificationInfo> PasteErrors;
-	TArray<FMovieSceneBindingProxy> OutBindings;
-	if (!FSequencerUtilities::PasteBindings(TextToImport, Sequencer.ToSharedRef(), PasteBindingsParams, OutBindings, PasteErrors))
+	if (!FSequencerUtilities::PasteBindings(TextToImport, Sequencer.ToSharedRef(), PasteBindingsParams, OutObjectBindings, PasteErrors))
 	{
 		for (FNotificationInfo PasteError : PasteErrors)
 		{
