Sequencer: Fix subsystem which was incorrectly marked as Blueprintable.

This was an unfortunate oversight and there was actually one case where a blueprint was set up with this parent class. I'm working with the author to fix this. There doesn't appear to be any kind of deprecation pathway for this kind of thing, so we'll need to fix anything by hand that comes up. Note, there's no need to subclass from this because you can just get the subsystem and call functions on it.

diff --git a/Engine/Plugins/MovieScene/LevelSequenceEditor/Source/LevelSequenceEditor/Public/LevelSequenceEditorSubsystem.h b/Engine/Plugins/MovieScene/LevelSequenceEditor/Source/LevelSequenceEditor/Public/LevelSequenceEditorSubsystem.h
--- a/Engine/Plugins/MovieScene/LevelSequenceEditor/Source/LevelSequenceEditor/Public/LevelSequenceEditorSubsystem.h
+++ b/Engine/Plugins/MovieScene/LevelSequenceEditor/Source/LevelSequenceEditor/Public/LevelSequenceEditorSubsystem.h
@@ -44,7 +44,7 @@ struct FMovieSceneScriptingParams
 * ULevelSequenceEditorSubsystem
 * Subsystem for level sequence editor related utilities to scripts
 */
-UCLASS(Blueprintable)
+UCLASS()
 class LEVELSEQUENCEEDITOR_API ULevelSequenceEditorSubsystem : public UEditorSubsystem
 {
 	GENERATED_BODY()
