Sequencer: Detach when creating a spawnable type from an actor.

The new spawnable should not be attached to anything. This fixes an issue where if you convert a child possessable to a spawnable, that it still thinks it's attached to the original parent.

diff --git a/Engine/Plugins/MovieScene/LevelSequenceEditor/Source/LevelSequenceEditor/Private/Misc/LevelSequenceEditorActorSpawner.cpp b/Engine/Plugins/MovieScene/LevelSequenceEditor/Source/LevelSequenceEditor/Private/Misc/LevelSequenceEditorActorSpawner.cpp
--- a/Engine/Plugins/MovieScene/LevelSequenceEditor/Source/LevelSequenceEditor/Private/Misc/LevelSequenceEditorActorSpawner.cpp
+++ b/Engine/Plugins/MovieScene/LevelSequenceEditor/Source/LevelSequenceEditor/Private/Misc/LevelSequenceEditorActorSpawner.cpp
@@ -48,6 +48,7 @@ TValueOrError<FNewSpawnable, FText> FLevelSequenceEditorActorSpawner::CreateNewS
 		}
 
 		AActor* SpawnedActor = Cast<AActor>(StaticDuplicateObject(Actor, &OwnerMovieScene, TemplateName, RF_AllFlags));
+		SpawnedActor->DetachFromActor(FDetachmentTransformRules(EDetachmentRule::KeepRelative, false));
 		SpawnedActor->bIsEditorPreviewActor = false;
 		NewSpawnable.ObjectTemplate = SpawnedActor;
 		NewSpawnable.Name = Actor->GetActorLabel();
