Sequencer: When baking, check that the bake data is valid. The bake operation bakes from start to end frame inclusive. In the case of a spawnable in a subsequence, the spawnable won't exist for the end frame, so there is no bake data, in which case we use the default transform.

diff --git a/Engine/Plugins/MovieScene/LevelSequenceEditor/Source/LevelSequenceEditor/Private/LevelSequenceEditorSubsystem.cpp b/Engine/Plugins/MovieScene/LevelSequenceEditor/Source/LevelSequenceEditor/Private/LevelSequenceEditorSubsystem.cpp
--- a/Engine/Plugins/MovieScene/LevelSequenceEditor/Source/LevelSequenceEditor/Private/LevelSequenceEditorSubsystem.cpp
+++ b/Engine/Plugins/MovieScene/LevelSequenceEditor/Source/LevelSequenceEditor/Private/LevelSequenceEditorSubsystem.cpp
@@ -1272,7 +1272,24 @@ bool ULevelSequenceEditorSubsystem::BakeTransformWithSettings(const TArray<FMovi
 
 			for (int32 Counter = 0; Counter < KeyTimes.Num(); ++Counter)
 			{
-				FTransform LocalTransform(BakeData.Value.Rotations[Counter], BakeData.Value.Locations[Counter], BakeData.Value.Scales[Counter]);
+				FVector LocalTranslation = DefaultLocation;
+				FVector LocalScale = DefaultScale;
+				FRotator LocalRotation = DefaultRotation.Rotation();
+
+				if (Counter < BakeData.Value.Locations.Num())
+				{
+					LocalTranslation = BakeData.Value.Locations[Counter];
+				}
+				if (Counter < BakeData.Value.Rotations.Num())
+				{
+					LocalRotation = BakeData.Value.Rotations[Counter];
+				}
+				if (Counter < BakeData.Value.Scales.Num())
+				{
+					LocalScale = BakeData.Value.Scales[Counter];
+				}
+
+				FTransform LocalTransform(LocalRotation, LocalTranslation, LocalScale);
 				LocalTranslations[Counter] = LocalTransform.GetTranslation();
 				LocalRotations[Counter] = LocalTransform.GetRotation().Euler();
 				LocalScales[Counter] = LocalTransform.GetScale3D();
