Sequencer: Fix Exclude Default Property Tracks array so that it can include properties on components multiple times. CameraComponent.FocusSettings.ManualFocusDistance, CameraComponent.CurrentAperture.

diff --git a/Engine/Plugins/MovieScene/LevelSequenceEditor/Source/LevelSequenceEditor/Private/LevelSequenceEditorToolkit.cpp b/Engine/Plugins/MovieScene/LevelSequenceEditor/Source/LevelSequenceEditor/Private/LevelSequenceEditorToolkit.cpp
--- a/Engine/Plugins/MovieScene/LevelSequenceEditor/Source/LevelSequenceEditor/Private/LevelSequenceEditorToolkit.cpp
+++ b/Engine/Plugins/MovieScene/LevelSequenceEditor/Source/LevelSequenceEditor/Private/LevelSequenceEditorToolkit.cpp
@@ -526,8 +526,7 @@ void FLevelSequenceEditorToolkit::AddDefaultTracksForActor(AActor& Actor, const
 				{
 					TArray<FString> PropertyNames;
 					PropertyTrackSettings.PropertyPath.ParseIntoArray(PropertyNames, TEXT("."));
-
-					ExcludePropertyTracksMap.Add(PropertyOwner, PropertyNames);
+					ExcludePropertyTracksMap.FindOrAdd(PropertyOwner).Append(PropertyNames);
 				}
 			}
 		}
