Also checking parent class on component validation.

diff --git a/Engine/Plugins/MovieScene/LevelSequenceEditor/Source/LevelSequenceEditor/Private/LevelSequenceEditorSubsystem.cpp b/Engine/Plugins/MovieScene/LevelSequenceEditor/Source/LevelSequenceEditor/Private/LevelSequenceEditorSubsystem.cpp
--- a/Engine/Plugins/MovieScene/LevelSequenceEditor/Source/LevelSequenceEditor/Private/LevelSequenceEditorSubsystem.cpp
+++ b/Engine/Plugins/MovieScene/LevelSequenceEditor/Source/LevelSequenceEditor/Private/LevelSequenceEditorSubsystem.cpp
@@ -1839,6 +1839,12 @@ void ULevelSequenceEditorSubsystem::GetRebindComponentNames(TArray<FName>& OutCo
 			if (GlobalClassFilter.IsValid())
 			{
 				bValidComponent = GlobalClassFilter->IsClassAllowed(ClassViewerOptions, Component->GetClass(), ClassFilterFuncs);
+
+				// Allow listing component if their parent class is allowed
+				while (!bValidComponent && Component->GetClass()->GetSuperClass())
+				{
+					bValidComponent = GlobalClassFilter->IsClassAllowed(ClassViewerOptions, Component->GetClass()->GetSuperClass(), ClassFilterFuncs);
+				}
 			}
 
 			if (bValidComponent)
diff --git a/Engine/Plugins/MovieScene/LevelSequenceEditor/Source/LevelSequenceEditor/Private/LevelSequenceEditorToolkit.cpp b/Engine/Plugins/MovieScene/LevelSequenceEditor/Source/LevelSequenceEditor/Private/LevelSequenceEditorToolkit.cpp
--- a/Engine/Plugins/MovieScene/LevelSequenceEditor/Source/LevelSequenceEditor/Private/LevelSequenceEditorToolkit.cpp
+++ b/Engine/Plugins/MovieScene/LevelSequenceEditor/Source/LevelSequenceEditor/Private/LevelSequenceEditorToolkit.cpp
@@ -948,11 +948,17 @@ void FLevelSequenceEditorToolkit::HandleTrackMenuExtensionAddTrack(FMenuBuilder&
 			{
 				if (Component)
 				{
-					bool bValidComponent = Component && !Component->IsVisualizationComponent();
+					bool bValidComponent = !Component->IsVisualizationComponent();
 
 					if (GlobalClassFilter.IsValid())
 					{
 						bValidComponent = GlobalClassFilter->IsClassAllowed(ClassViewerOptions, Component->GetClass(), ClassFilterFuncs);
+						
+						// Allow listing component if their parent class is allowed
+						while (!bValidComponent && Component->GetClass()->GetSuperClass())
+						{
+							bValidComponent = GlobalClassFilter->IsClassAllowed(ClassViewerOptions, Component->GetClass()->GetSuperClass(), ClassFilterFuncs);
+						}
 					}
 
 					if (bValidComponent)
