Fixed regression introduced with 26827605

diff --git a/Engine/Plugins/MovieScene/LevelSequenceEditor/Source/LevelSequenceEditor/Private/LevelSequenceEditorSubsystem.cpp b/Engine/Plugins/MovieScene/LevelSequenceEditor/Source/LevelSequenceEditor/Private/LevelSequenceEditorSubsystem.cpp
--- a/Engine/Plugins/MovieScene/LevelSequenceEditor/Source/LevelSequenceEditor/Private/LevelSequenceEditorSubsystem.cpp
+++ b/Engine/Plugins/MovieScene/LevelSequenceEditor/Source/LevelSequenceEditor/Private/LevelSequenceEditorSubsystem.cpp
@@ -1841,9 +1841,11 @@ void ULevelSequenceEditorSubsystem::GetRebindComponentNames(TArray<FName>& OutCo
 				bValidComponent = GlobalClassFilter->IsClassAllowed(ClassViewerOptions, Component->GetClass(), ClassFilterFuncs);
 
 				// Allow listing component if their parent class is allowed
-				while (!bValidComponent && Component->GetClass()->GetSuperClass())
+				UClass* ComponentSuperClass = Component->GetClass()->GetSuperClass();
+				while (!bValidComponent && ComponentSuperClass)
 				{
-					bValidComponent = GlobalClassFilter->IsClassAllowed(ClassViewerOptions, Component->GetClass()->GetSuperClass(), ClassFilterFuncs);
+					bValidComponent = GlobalClassFilter->IsClassAllowed(ClassViewerOptions, ComponentSuperClass, ClassFilterFuncs);
+					ComponentSuperClass = ComponentSuperClass->GetSuperClass();
 				}
 			}
 
diff --git a/Engine/Plugins/MovieScene/LevelSequenceEditor/Source/LevelSequenceEditor/Private/LevelSequenceEditorToolkit.cpp b/Engine/Plugins/MovieScene/LevelSequenceEditor/Source/LevelSequenceEditor/Private/LevelSequenceEditorToolkit.cpp
--- a/Engine/Plugins/MovieScene/LevelSequenceEditor/Source/LevelSequenceEditor/Private/LevelSequenceEditorToolkit.cpp
+++ b/Engine/Plugins/MovieScene/LevelSequenceEditor/Source/LevelSequenceEditor/Private/LevelSequenceEditorToolkit.cpp
@@ -955,9 +955,11 @@ void FLevelSequenceEditorToolkit::HandleTrackMenuExtensionAddTrack(FMenuBuilder&
 						bValidComponent = GlobalClassFilter->IsClassAllowed(ClassViewerOptions, Component->GetClass(), ClassFilterFuncs);
 						
 						// Allow listing component if their parent class is allowed
-						while (!bValidComponent && Component->GetClass()->GetSuperClass())
+						UClass* ComponentSuperClass = Component->GetClass()->GetSuperClass();
+						while (!bValidComponent && ComponentSuperClass)
 						{
-							bValidComponent = GlobalClassFilter->IsClassAllowed(ClassViewerOptions, Component->GetClass()->GetSuperClass(), ClassFilterFuncs);
+							bValidComponent = GlobalClassFilter->IsClassAllowed(ClassViewerOptions, ComponentSuperClass, ClassFilterFuncs);
+							ComponentSuperClass = ComponentSuperClass->GetSuperClass();
 						}
 					}
 
