Sequencer: SetCurrentTime and SetCurrentLocalTime should force an evaluation if the current time is already the requested time. Users would expect an evaluation to occur.

diff --git a/Engine/Plugins/MovieScene/LevelSequenceEditor/Source/LevelSequenceEditor/Private/LevelSequenceEditorBlueprintLibrary.cpp b/Engine/Plugins/MovieScene/LevelSequenceEditor/Source/LevelSequenceEditor/Private/LevelSequenceEditorBlueprintLibrary.cpp
--- a/Engine/Plugins/MovieScene/LevelSequenceEditor/Source/LevelSequenceEditor/Private/LevelSequenceEditorBlueprintLibrary.cpp
+++ b/Engine/Plugins/MovieScene/LevelSequenceEditor/Source/LevelSequenceEditor/Private/LevelSequenceEditorBlueprintLibrary.cpp
@@ -126,7 +126,15 @@ void ULevelSequenceEditorBlueprintLibrary::SetCurrentTime(int32 NewFrame)
 		FFrameRate DisplayRate = CurrentSequencer.Pin()->GetFocusedDisplayRate();
 		FFrameRate TickResolution = CurrentSequencer.Pin()->GetFocusedTickResolution();
 
-		CurrentSequencer.Pin()->SetGlobalTime(ConvertFrameTime(NewFrame, DisplayRate, TickResolution));
+		FFrameTime GlobalTime = ConvertFrameTime(NewFrame, DisplayRate, TickResolution);
+		if (GlobalTime == CurrentSequencer.Pin()->GetGlobalTime().Time)
+		{
+			CurrentSequencer.Pin()->ForceEvaluate();
+		}
+		else
+		{
+			CurrentSequencer.Pin()->SetGlobalTime(GlobalTime);
+		}
 	}
 }
 
@@ -149,7 +157,15 @@ void ULevelSequenceEditorBlueprintLibrary::SetCurrentLocalTime(int32 NewFrame)
 		FFrameRate DisplayRate = CurrentSequencer.Pin()->GetFocusedDisplayRate();
 		FFrameRate TickResolution = CurrentSequencer.Pin()->GetFocusedTickResolution();
 
-		CurrentSequencer.Pin()->SetLocalTime(ConvertFrameTime(NewFrame, DisplayRate, TickResolution));
+		FFrameTime LocalTime = ConvertFrameTime(NewFrame, DisplayRate, TickResolution);
+		if (LocalTime == CurrentSequencer.Pin()->GetLocalTime().Time)
+		{
+			CurrentSequencer.Pin()->ForceEvaluate();
+		}
+		else
+		{
+			CurrentSequencer.Pin()->SetLocalTime(LocalTime);
+		}
 	}
 }
 
diff --git a/Engine/Plugins/MovieScene/LevelSequenceEditor/Source/LevelSequenceEditor/Public/LevelSequenceEditorBlueprintLibrary.h b/Engine/Plugins/MovieScene/LevelSequenceEditor/Source/LevelSequenceEditor/Public/LevelSequenceEditorBlueprintLibrary.h
--- a/Engine/Plugins/MovieScene/LevelSequenceEditor/Source/LevelSequenceEditor/Public/LevelSequenceEditorBlueprintLibrary.h
+++ b/Engine/Plugins/MovieScene/LevelSequenceEditor/Source/LevelSequenceEditor/Public/LevelSequenceEditorBlueprintLibrary.h
@@ -102,7 +102,7 @@ public:
 public:
 
 	/**
-	 * Set global playback position for the current level sequence in frames
+	 * Set global playback position for the current level sequence in frames. If the requested time is the same as the current time, an evaluation will be forced.
 	 */
 	UFUNCTION(BlueprintCallable, Category = "Level Sequence Editor")
 	static void SetCurrentTime(int32 NewFrame);
@@ -114,7 +114,7 @@ public:
 	static int32 GetCurrentTime();
 
 	/**
-	 * Set local playback position for the current level sequence in frames
+	 * Set local playback position for the current level sequence in frames. If the requested time is the same as the current time, an evaluation will be forced.
 	 */
 	UFUNCTION(BlueprintCallable, Category = "Level Sequence Editor")
 	static void SetCurrentLocalTime(int32 NewFrame);
