Sequencer: Fix SelectKeys when invalid handle

#rnx

diff --git a/Engine/Plugins/MovieScene/LevelSequenceEditor/Source/LevelSequenceEditor/Private/LevelSequenceEditorBlueprintLibrary.cpp b/Engine/Plugins/MovieScene/LevelSequenceEditor/Source/LevelSequenceEditor/Private/LevelSequenceEditorBlueprintLibrary.cpp
--- a/Engine/Plugins/MovieScene/LevelSequenceEditor/Source/LevelSequenceEditor/Private/LevelSequenceEditorBlueprintLibrary.cpp
+++ b/Engine/Plugins/MovieScene/LevelSequenceEditor/Source/LevelSequenceEditor/Private/LevelSequenceEditorBlueprintLibrary.cpp
@@ -414,8 +414,11 @@ void ULevelSequenceEditorBlueprintLibrary::SelectKeys(const FSequencerChannelPro
 							FKeySelection& KeySelection = CurrentSequencer.Pin()->GetViewModel()->GetSelection()->KeySelection;
 							for (int32 Index : Indices)
 							{
-								FKeyHandle KeyHandle = MovieSceneChannel->GetHandle(Index);
-								KeySelection.Select(ChannelModel, KeyHandle);
+								if (Index >= 0 && Index < MovieSceneChannel->GetNumKeys())
+								{
+									FKeyHandle KeyHandle = MovieSceneChannel->GetHandle(Index);
+									KeySelection.Select(ChannelModel, KeyHandle);
+								}
 							}
 							break;
 						}
