Sequencer: Fixup some more assumptions where the playback context was a UWorld

diff --git a/Engine/Plugins/MovieScene/LevelSequenceEditor/Source/LevelSequenceEditor/Private/Misc/LevelSequenceCustomization.cpp b/Engine/Plugins/MovieScene/LevelSequenceEditor/Source/LevelSequenceEditor/Private/Misc/LevelSequenceCustomization.cpp
--- a/Engine/Plugins/MovieScene/LevelSequenceEditor/Source/LevelSequenceEditor/Private/Misc/LevelSequenceCustomization.cpp
+++ b/Engine/Plugins/MovieScene/LevelSequenceEditor/Source/LevelSequenceEditor/Private/Misc/LevelSequenceCustomization.cpp
@@ -384,7 +384,7 @@ void FLevelSequenceCustomization::AddSpawnLevelMenu(FMenuBuilder& MenuBuilder, T
 		EUserInterfaceActionType::ToggleButton
 	);
 
-	UWorld* World = Cast<UWorld>(Sequencer->GetPlaybackContext());
+	UWorld* World = Sequencer->GetPlaybackContext()->GetWorld();
 	if (!World)
 	{
 		return;
diff --git a/Engine/Plugins/MovieScene/LevelSequenceEditor/Source/LevelSequenceEditor/Private/Misc/LevelSequenceEditorSpawnRegister.cpp b/Engine/Plugins/MovieScene/LevelSequenceEditor/Source/LevelSequenceEditor/Private/Misc/LevelSequenceEditorSpawnRegister.cpp
--- a/Engine/Plugins/MovieScene/LevelSequenceEditor/Source/LevelSequenceEditor/Private/Misc/LevelSequenceEditorSpawnRegister.cpp
+++ b/Engine/Plugins/MovieScene/LevelSequenceEditor/Source/LevelSequenceEditor/Private/Misc/LevelSequenceEditorSpawnRegister.cpp
@@ -363,7 +363,7 @@ void FLevelSequenceEditorSpawnRegister::HandleConvertPossessableToSpawnable(UObj
 		}
 
 		GEditor->SelectActor(OldActor, false, true);
-		UObject* PlaybackContext = SharedPlaybackState->GetPlaybackContext()->GetWorld();
+		UObject* PlaybackContext = SharedPlaybackState->GetPlaybackContext();
 		UWorld* World = PlaybackContext ? PlaybackContext->GetWorld() : nullptr;
 		if (World)
 		{
diff --git a/Engine/Source/Editor/MovieSceneTools/Private/TrackEditors/LevelVisibilityTrackEditor.cpp b/Engine/Source/Editor/MovieSceneTools/Private/TrackEditors/LevelVisibilityTrackEditor.cpp
--- a/Engine/Source/Editor/MovieSceneTools/Private/TrackEditors/LevelVisibilityTrackEditor.cpp
+++ b/Engine/Source/Editor/MovieSceneTools/Private/TrackEditors/LevelVisibilityTrackEditor.cpp
@@ -190,7 +190,7 @@ void FLevelVisibilityTrackEditor::OnAddNewSection( UMovieSceneTrack* LevelVisibi
 
 void FLevelVisibilityTrackEditor::GetCurrentLevelVisibility(TArray<FName>& OutVisibleLevelNames, TArray<FName>& OutHiddenLevelNames)
 {
-	UWorld* World = Cast<UWorld>(GetSequencer()->GetPlaybackContext());
+	UWorld* World = GetSequencer()->GetPlaybackContext()->GetWorld();
 	if (!World)
 	{
 		return;
