Sequencer: UX changes for Playback Controls

Update a few more icons
Fix twitchy behavior in displaying the current time with subframes

diff --git a/Engine/Plugins/MovieScene/LevelSequenceEditor/Source/LevelSequenceEditor/Private/CinematicViewport/SCinematicLevelViewport.cpp b/Engine/Plugins/MovieScene/LevelSequenceEditor/Source/LevelSequenceEditor/Private/CinematicViewport/SCinematicLevelViewport.cpp
--- a/Engine/Plugins/MovieScene/LevelSequenceEditor/Source/LevelSequenceEditor/Private/CinematicViewport/SCinematicLevelViewport.cpp
+++ b/Engine/Plugins/MovieScene/LevelSequenceEditor/Source/LevelSequenceEditor/Private/CinematicViewport/SCinematicLevelViewport.cpp
@@ -371,6 +371,7 @@ void SCinematicLevelViewport::Construct(const FArguments& InArgs)
 								.ColorAndOpacity(Gray)
 								.Text_Lambda([this] { return UIData.LocalPlaybackTime; })
 								.ToolTipText(LOCTEXT("LocalPlaybackTime", "The current playback time relative to the currently evaluated sequence."))
+								.MinDesiredWidth(this, &SCinematicLevelViewport::GetPlayTimeMinDesiredWidth)
 							]
 						]
 					]
