Sequencer- Fix uninitialized struct variables.

diff --git a/Engine/Plugins/MovieScene/LevelSequenceEditor/Source/LevelSequenceEditor/Public/LevelSequenceEditorSubsystem.h b/Engine/Plugins/MovieScene/LevelSequenceEditor/Source/LevelSequenceEditor/Public/LevelSequenceEditorSubsystem.h
--- a/Engine/Plugins/MovieScene/LevelSequenceEditor/Source/LevelSequenceEditor/Public/LevelSequenceEditorSubsystem.h
+++ b/Engine/Plugins/MovieScene/LevelSequenceEditor/Source/LevelSequenceEditor/Public/LevelSequenceEditorSubsystem.h
@@ -59,11 +59,11 @@ struct FMovieSceneUniversalLocatorInfo
 
 	// Flags for how to resolve the locator when in editor
 	UPROPERTY(EditAnywhere, Category = "Default")
-	ELocatorResolveFlags EditorResolveFlags;
+	ELocatorResolveFlags EditorResolveFlags = ELocatorResolveFlags::None;
 
 	// Flags for how to resolve the locator when in runtime
 	UPROPERTY(EditAnywhere, Category = "Default")
-	ELocatorResolveFlags RuntimeResolveFlags;
+	ELocatorResolveFlags RuntimeResolveFlags = ELocatorResolveFlags::None;
 };
 
 // Helper struct for editing arrays of locators for object bindings
diff --git a/Engine/Source/Runtime/MovieScene/Public/MovieSceneBindingReferences.h b/Engine/Source/Runtime/MovieScene/Public/MovieSceneBindingReferences.h
--- a/Engine/Source/Runtime/MovieScene/Public/MovieSceneBindingReferences.h
+++ b/Engine/Source/Runtime/MovieScene/Public/MovieSceneBindingReferences.h
@@ -27,11 +27,11 @@ struct FMovieSceneBindingReference
 
 #if WITH_EDITORONLY_DATA
 	UPROPERTY()
-	ELocatorResolveFlags EditorResolveFlags;
+	ELocatorResolveFlags EditorResolveFlags = ELocatorResolveFlags::None;
 #endif
 
 	UPROPERTY()
-	ELocatorResolveFlags RuntimeResolveFlags;
+	ELocatorResolveFlags RuntimeResolveFlags = ELocatorResolveFlags::None;
 
 	void InitializeLocatorResolveFlags();
 };
