Sequencer- Cleanup and minor refactor of use of locators for Sequencer in preparation for the 5.4/29.00 branch.
* Removed Editor vs Runtime flags for locator use in Sequencer. Long term we're planning on using UObject based bindings instead for allowing editor preview rather than locators.
* Remove the auto-addition of the Binding Lifetime track for now, as it should only be applicable for Replaceable and Spawnables.
* Change the UI for locator type selection to only allow Actor types for now.

diff --git a/Engine/Plugins/MovieScene/LevelSequenceEditor/Source/LevelSequenceEditor/Private/LevelSequenceEditorSubsystem.cpp b/Engine/Plugins/MovieScene/LevelSequenceEditor/Source/LevelSequenceEditor/Private/LevelSequenceEditorSubsystem.cpp
--- a/Engine/Plugins/MovieScene/LevelSequenceEditor/Source/LevelSequenceEditor/Private/LevelSequenceEditorSubsystem.cpp
+++ b/Engine/Plugins/MovieScene/LevelSequenceEditor/Source/LevelSequenceEditor/Private/LevelSequenceEditorSubsystem.cpp
@@ -1842,7 +1842,7 @@ void ULevelSequenceEditorSubsystem::AddBindingPropertiesMenu(FMenuBuilder& MenuB
 		FMovieSceneUniversalLocatorList* Locators = (FMovieSceneUniversalLocatorList*)LocatorsStruct->GetStructMemory();
 		Algo::Transform(BindingReferences->GetReferences(ObjectBindings[0]), Locators->Bindings, [](const FMovieSceneBindingReference& Reference) 
 			{ 
-				return FMovieSceneUniversalLocatorInfo{ Reference.Locator, Reference.EditorResolveFlags, Reference.RuntimeResolveFlags };
+				return FMovieSceneUniversalLocatorInfo{ Reference.Locator, Reference.ResolveFlags };
 			});
 
 		MenuBuilder.AddMenuSeparator();
@@ -1942,15 +1942,12 @@ void ULevelSequenceEditorSubsystem::OnFinishedChangingLocators(const FPropertyCh
 				const FMovieSceneBindingReference* NewRef = BindingReferences->AddBinding(ObjectBindingID, MoveTemp(LocatorInfo.Locator));
 				if (NewRef)
 				{
-#if WITH_EDITOR
-					LocatorInfo.EditorResolveFlags = NewRef->EditorResolveFlags;
-#endif
-					LocatorInfo.RuntimeResolveFlags = NewRef->RuntimeResolveFlags;
+					LocatorInfo.ResolveFlags = NewRef->ResolveFlags;
 				}
 			}
 			else
 			{
-				BindingReferences->AddBinding(ObjectBindingID, MoveTemp(LocatorInfo.Locator), LocatorInfo.EditorResolveFlags, LocatorInfo.RuntimeResolveFlags);
+				BindingReferences->AddBinding(ObjectBindingID, MoveTemp(LocatorInfo.Locator), LocatorInfo.ResolveFlags);
 			}
 		}
 
@@ -2011,7 +2008,7 @@ void ULevelSequenceEditorSubsystem::OnFinishedChangingLocators(const FPropertyCh
 		Locators->Bindings.Empty();
 		Algo::Transform(BindingReferences->GetReferences(ObjectBindingID), Locators->Bindings, [](const FMovieSceneBindingReference& Reference)
 			{
-				return FMovieSceneUniversalLocatorInfo{ Reference.Locator, Reference.EditorResolveFlags, Reference.RuntimeResolveFlags };
+				return FMovieSceneUniversalLocatorInfo{ Reference.Locator, Reference.ResolveFlags };
 			});
 
 
diff --git a/Engine/Plugins/MovieScene/LevelSequenceEditor/Source/LevelSequenceEditor/Public/LevelSequenceEditorSubsystem.h b/Engine/Plugins/MovieScene/LevelSequenceEditor/Source/LevelSequenceEditor/Public/LevelSequenceEditorSubsystem.h
--- a/Engine/Plugins/MovieScene/LevelSequenceEditor/Source/LevelSequenceEditor/Public/LevelSequenceEditorSubsystem.h
+++ b/Engine/Plugins/MovieScene/LevelSequenceEditor/Source/LevelSequenceEditor/Public/LevelSequenceEditorSubsystem.h
@@ -54,16 +54,12 @@ struct FMovieSceneUniversalLocatorInfo
 	GENERATED_BODY()
 
 	// Locator for the entry
-	UPROPERTY(EditAnywhere, Category = "Default")
+	UPROPERTY(EditAnywhere, Category = "Default", meta=(AllowedLocators="Actor"))
 	FUniversalObjectLocator Locator;
 
-	// Flags for how to resolve the locator when in editor
-	UPROPERTY(EditAnywhere, Category = "Default")
-	ELocatorResolveFlags EditorResolveFlags = ELocatorResolveFlags::None;
-
-	// Flags for how to resolve the locator when in runtime
+	// Flags for how to resolve the locator
 	UPROPERTY(EditAnywhere, Category = "Default")
-	ELocatorResolveFlags RuntimeResolveFlags = ELocatorResolveFlags::None;
+	ELocatorResolveFlags ResolveFlags = ELocatorResolveFlags::None;
 };
 
 // Helper struct for editing arrays of locators for object bindings
diff --git a/Engine/Source/Editor/Sequencer/Private/Sequencer.cpp b/Engine/Source/Editor/Sequencer/Private/Sequencer.cpp
--- a/Engine/Source/Editor/Sequencer/Private/Sequencer.cpp
+++ b/Engine/Source/Editor/Sequencer/Private/Sequencer.cpp
@@ -8893,41 +8893,6 @@ void FixSortingOrders(FMovieSceneBinding* InBinding, UMovieScene* MovieScene)
 
 void FSequencer::OnAddBinding(const FGuid& ObjectBinding, UMovieScene* MovieScene)
 {
-	// If a new binding requires a binding lifetime track and doesn't have one, add one.
-	if (UMovieSceneSequence* OuterSequence = MovieScene->GetTypedOuter<UMovieSceneSequence>())
-	{
-		if (FMovieSceneBindingReferences* BindingReferences = OuterSequence->GetBindingReferences())
-		{
-			TArrayView<const FMovieSceneBindingReference> References = BindingReferences->GetReferences(ObjectBinding);
-			bool bRequiresBindingLifetimeTrack = false;
-			for (const FMovieSceneBindingReference& Reference : References)
-			{
-				if (EnumHasAllFlags(Reference.EditorResolveFlags, ELocatorResolveFlags::Load) || EnumHasAllFlags(Reference.RuntimeResolveFlags, ELocatorResolveFlags::Load))
-				{
-					bRequiresBindingLifetimeTrack = true;
-					break;
-				}
-			}
-
-			// We may need to force-create a Binding Lifetime Track here if the locator has been set to load. This allows the sequence to manage load/unload of the locator.
-			if (bRequiresBindingLifetimeTrack)
-			{
-				UMovieSceneBindingLifetimeTrack* BindingLifetimeTrack = Cast<UMovieSceneBindingLifetimeTrack>(MovieScene->FindTrack(UMovieSceneBindingLifetimeTrack::StaticClass(), ObjectBinding, NAME_None));
-				if (!BindingLifetimeTrack)
-				{
-					BindingLifetimeTrack = Cast<UMovieSceneBindingLifetimeTrack>(MovieScene->AddTrack(UMovieSceneBindingLifetimeTrack::StaticClass(), ObjectBinding));
-				}
-
-				if (BindingLifetimeTrack && BindingLifetimeTrack->GetAllSections().IsEmpty())
-				{
-					UMovieSceneBindingLifetimeSection* BindingLifetimeSection = Cast<UMovieSceneBindingLifetimeSection>(BindingLifetimeTrack->CreateNewSection());
-					BindingLifetimeSection->SetRange(TRange<FFrameNumber>::All());
-					BindingLifetimeTrack->AddSection(*BindingLifetimeSection);
-				}
-			}
-		}
-	}
-
 	FMovieSceneBinding* Binding = MovieScene->FindBinding(ObjectBinding);
 	if (Binding)
 	{
diff --git a/Engine/Source/Runtime/Engine/Private/UnrealEngine.cpp b/Engine/Source/Runtime/Engine/Private/UnrealEngine.cpp
--- a/Engine/Source/Runtime/Engine/Private/UnrealEngine.cpp
+++ b/Engine/Source/Runtime/Engine/Private/UnrealEngine.cpp
@@ -329,8 +329,7 @@ void FEngineModule::StartupModule()
 			{
 				FFragmentTypeParameters FragmentTypeParams("asset", NSLOCTEXT("Engine", "AssetLocatorFragment", "Asset"));
 				FragmentTypeParams.PrimaryEditorType = "Asset";
-				FragmentTypeParams.EditorFlags = EFragmentTypeFlags::CanBeLoaded | EFragmentTypeFlags::LoadedByDefault;
-				FragmentTypeParams.RuntimeFlags = EFragmentTypeFlags::CanBeLoaded;
+				FragmentTypeParams.Flags = EFragmentTypeFlags::CanBeLoaded;
 				FAssetLocatorFragment::FragmentType = UolModule.RegisterFragmentType<FAssetLocatorFragment>(FragmentTypeParams);
 			}
 			{
diff --git a/Engine/Source/Runtime/MovieScene/Private/MovieSceneBindingReferences.cpp b/Engine/Source/Runtime/MovieScene/Private/MovieSceneBindingReferences.cpp
--- a/Engine/Source/Runtime/MovieScene/Private/MovieSceneBindingReferences.cpp
+++ b/Engine/Source/Runtime/MovieScene/Private/MovieSceneBindingReferences.cpp
@@ -68,15 +68,12 @@ const FMovieSceneBindingReference* FMovieSceneBindingReferences::AddBinding(cons
 	return &NewBinding;
 }
 
-const FMovieSceneBindingReference* FMovieSceneBindingReferences::AddBinding(const FGuid& ObjectId, FUniversalObjectLocator&& NewLocator, ELocatorResolveFlags InEditorResolveFlags, ELocatorResolveFlags InRuntimeResolveFlags)
+const FMovieSceneBindingReference* FMovieSceneBindingReferences::AddBinding(const FGuid& ObjectId, FUniversalObjectLocator&& NewLocator, ELocatorResolveFlags InResolveFlags)
 {
 	const int32 Index = Algo::UpperBoundBy(SortedReferences, ObjectId, &FMovieSceneBindingReference::ID);
 
 	FMovieSceneBindingReference& NewBinding = SortedReferences.Insert_GetRef(FMovieSceneBindingReference{ ObjectId, MoveTemp(NewLocator) }, Index);
-#if WITH_EDITOR
-	NewBinding.EditorResolveFlags = InEditorResolveFlags;
-#endif
-	NewBinding.RuntimeResolveFlags = InRuntimeResolveFlags;
+	NewBinding.ResolveFlags = InResolveFlags;
 
 	return &NewBinding;
 }
@@ -119,16 +116,7 @@ void FMovieSceneBindingReferences::ResolveBinding(const FGuid& ObjectId, const U
 		{
 			if (UWorld* World = ResolveParams.Context->GetWorld())
 			{
-#if WITH_EDITOR
-				if (World->IsEditorWorld())
-				{
-					EnumAddFlags(const_cast<UE::UniversalObjectLocator::FResolveParams&>(ResolveParams).Flags, SortedReferences[Index].EditorResolveFlags);
-				}
-				else
-#endif
-				{
-					EnumAddFlags(const_cast<UE::UniversalObjectLocator::FResolveParams&>(ResolveParams).Flags, SortedReferences[Index].RuntimeResolveFlags);
-				}
+				EnumAddFlags(const_cast<UE::UniversalObjectLocator::FResolveParams&>(ResolveParams).Flags, SortedReferences[Index].ResolveFlags);
 			}
 		}
 
@@ -237,8 +225,5 @@ void FMovieSceneBindingReference::InitializeLocatorResolveFlags()
 		}
 	};
 
-#if WITH_EDITOR
-	InitializeFlags(Locator.GetDefaultEditorFlags(), EditorResolveFlags);
-#endif
-	InitializeFlags(Locator.GetDefaultRuntimeFlags(), RuntimeResolveFlags);
+	InitializeFlags(Locator.GetDefaultFlags(), ResolveFlags);
 }
diff --git a/Engine/Source/Runtime/MovieScene/Public/MovieSceneBindingReferences.h b/Engine/Source/Runtime/MovieScene/Public/MovieSceneBindingReferences.h
--- a/Engine/Source/Runtime/MovieScene/Public/MovieSceneBindingReferences.h
+++ b/Engine/Source/Runtime/MovieScene/Public/MovieSceneBindingReferences.h
@@ -25,13 +25,8 @@ struct FMovieSceneBindingReference
 	UPROPERTY()
 	FUniversalObjectLocator Locator;
 
-#if WITH_EDITORONLY_DATA
 	UPROPERTY()
-	ELocatorResolveFlags EditorResolveFlags = ELocatorResolveFlags::None;
-#endif
-
-	UPROPERTY()
-	ELocatorResolveFlags RuntimeResolveFlags = ELocatorResolveFlags::None;
+	ELocatorResolveFlags ResolveFlags = ELocatorResolveFlags::None;
 
 	void InitializeLocatorResolveFlags();
 };
@@ -93,7 +88,7 @@ struct FMovieSceneBindingReferences
 	 * @param ObjectId	The ID to associate the object with
 	 * @param InContext	A context in which InObject resides (either a UWorld, or an AActor)
 	 */
-	MOVIESCENE_API const FMovieSceneBindingReference* AddBinding(const FGuid& ObjectId, FUniversalObjectLocator&& NewLocator, ELocatorResolveFlags InEditorResolveFlags, ELocatorResolveFlags InRuntimeResolveFlags);
+	MOVIESCENE_API const FMovieSceneBindingReference* AddBinding(const FGuid& ObjectId, FUniversalObjectLocator&& NewLocator, ELocatorResolveFlags InResolveFlags);
 
 	/**
 	 * Resolve a binding for the specified ID using a given context
diff --git a/Engine/Source/Runtime/UniversalObjectLocator/Private/UniversalObjectLocator.cpp b/Engine/Source/Runtime/UniversalObjectLocator/Private/UniversalObjectLocator.cpp
--- a/Engine/Source/Runtime/UniversalObjectLocator/Private/UniversalObjectLocator.cpp
+++ b/Engine/Source/Runtime/UniversalObjectLocator/Private/UniversalObjectLocator.cpp
@@ -309,7 +309,7 @@ const FUniversalObjectLocatorFragment* FUniversalObjectLocator::GetLastFragment(
 }
 
 
-UE::UniversalObjectLocator::EFragmentTypeFlags FUniversalObjectLocator::GetDefaultEditorFlags() const
+UE::UniversalObjectLocator::EFragmentTypeFlags FUniversalObjectLocator::GetDefaultFlags() const
 {
 	using namespace UE::UniversalObjectLocator;
 	EFragmentTypeFlags Flags = EFragmentTypeFlags::None;
@@ -318,22 +318,7 @@ UE::UniversalObjectLocator::EFragmentTypeFlags FUniversalObjectLocator::GetDefau
 	{
 		if (const FFragmentType* FragmentTypePtr = Fragment.GetFragmentType())
 		{
-			EnumAddFlags(Flags, FragmentTypePtr->EditorFlags);
-		}
-	}
-	return Flags;
-}
-
-
-UE::UniversalObjectLocator::EFragmentTypeFlags FUniversalObjectLocator::GetDefaultRuntimeFlags() const
-{
-	using namespace UE::UniversalObjectLocator;
-	EFragmentTypeFlags Flags = EFragmentTypeFlags::None;
-	for (const FUniversalObjectLocatorFragment& Fragment : Fragments)
-	{
-		if (const FFragmentType* FragmentTypePtr = Fragment.GetFragmentType())
-		{
-			EnumAddFlags(Flags, FragmentTypePtr->RuntimeFlags);
+			EnumAddFlags(Flags, FragmentTypePtr->Flags);
 		}
 	}
 	return Flags;
diff --git a/Engine/Source/Runtime/UniversalObjectLocator/Public/IUniversalObjectLocatorModule.h b/Engine/Source/Runtime/UniversalObjectLocator/Public/IUniversalObjectLocatorModule.h
--- a/Engine/Source/Runtime/UniversalObjectLocator/Public/IUniversalObjectLocatorModule.h
+++ b/Engine/Source/Runtime/UniversalObjectLocator/Public/IUniversalObjectLocatorModule.h
@@ -30,8 +30,7 @@ struct FFragmentTypeParameters
 	FText DisplayText;
 	FName FragmentTypeID;
 	FName PrimaryEditorType;
-	EFragmentTypeFlags EditorFlags;
-	EFragmentTypeFlags RuntimeFlags;
+	EFragmentTypeFlags Flags;
 };
 
 class IUniversalObjectLocatorModule
@@ -75,8 +74,7 @@ public:
 		FFragmentType NewFragmentType;
 		NewFragmentType.FragmentTypeID     = FragmentTypeParameters.FragmentTypeID;
 		NewFragmentType.PrimaryEditorType  = FragmentTypeParameters.PrimaryEditorType;
-		NewFragmentType.EditorFlags		   = FragmentTypeParameters.EditorFlags;
-		NewFragmentType.RuntimeFlags	   = FragmentTypeParameters.RuntimeFlags;
+		NewFragmentType.Flags			   = FragmentTypeParameters.Flags;
 		NewFragmentType.DebuggingAssistant = MakeShared<TFragmentTypeDebuggingAssistant<PayloadStructType>>();
 		NewFragmentType.PayloadType        = PayloadStructType::StaticStruct();
 
diff --git a/Engine/Source/Runtime/UniversalObjectLocator/Public/UniversalObjectLocator.h b/Engine/Source/Runtime/UniversalObjectLocator/Public/UniversalObjectLocator.h
--- a/Engine/Source/Runtime/UniversalObjectLocator/Public/UniversalObjectLocator.h
+++ b/Engine/Source/Runtime/UniversalObjectLocator/Public/UniversalObjectLocator.h
@@ -197,14 +197,9 @@ public:
 
 
 	/*
-	* Iterates over all fragments and combines their types' default editor flags.
+	* Iterates over all fragments and combines their types' default flags.
 	*/
-	UNIVERSALOBJECTLOCATOR_API UE::UniversalObjectLocator::EFragmentTypeFlags GetDefaultEditorFlags() const;
-
-	/*
-	* Iterates over all fragments and combines their types' default runtime flags.
-	*/
-	UNIVERSALOBJECTLOCATOR_API UE::UniversalObjectLocator::EFragmentTypeFlags GetDefaultRuntimeFlags() const;
+	UNIVERSALOBJECTLOCATOR_API UE::UniversalObjectLocator::EFragmentTypeFlags GetDefaultFlags() const;
 
 	/**
 	 * Equality comparison.
diff --git a/Engine/Source/Runtime/UniversalObjectLocator/Public/UniversalObjectLocatorFragmentType.h b/Engine/Source/Runtime/UniversalObjectLocator/Public/UniversalObjectLocatorFragmentType.h
--- a/Engine/Source/Runtime/UniversalObjectLocator/Public/UniversalObjectLocatorFragmentType.h
+++ b/Engine/Source/Runtime/UniversalObjectLocator/Public/UniversalObjectLocatorFragmentType.h
@@ -115,11 +115,8 @@ struct FFragmentType
 	/** Name of the primary editor type for this fragment type, defining how it appears on UI */
 	FName PrimaryEditorType;
 
-	/** Flags defining behavior of this fragment type in editor */
-	EFragmentTypeFlags EditorFlags;
-
-	/** Flags defining behavior of this fragment type at runtime */
-	EFragmentTypeFlags RuntimeFlags;
+	/** Flags defining behavior of this fragment type */
+	EFragmentTypeFlags Flags;
 
 	UScriptStruct* GetStruct() const
 	{
