Sequencer: Add SetLoopMode/GetLoopMode

diff --git a/Engine/Plugins/MovieScene/LevelSequenceEditor/Source/LevelSequenceEditor/Private/LevelSequenceEditorBlueprintLibrary.cpp b/Engine/Plugins/MovieScene/LevelSequenceEditor/Source/LevelSequenceEditor/Private/LevelSequenceEditorBlueprintLibrary.cpp
--- a/Engine/Plugins/MovieScene/LevelSequenceEditor/Source/LevelSequenceEditor/Private/LevelSequenceEditorBlueprintLibrary.cpp
+++ b/Engine/Plugins/MovieScene/LevelSequenceEditor/Source/LevelSequenceEditor/Private/LevelSequenceEditorBlueprintLibrary.cpp
@@ -207,6 +207,24 @@ float ULevelSequenceEditorBlueprintLibrary::GetPlaybackSpeed()
 	return 0.f;
 }
 
+void ULevelSequenceEditorBlueprintLibrary::SetLoopMode(ESequencerLoopMode NewLoopMode)
+{
+	if (CurrentSequencer.IsValid())
+	{
+		CurrentSequencer.Pin()->GetSequencerSettings()->SetLoopMode(NewLoopMode);
+	}
+}
+
+ESequencerLoopMode ULevelSequenceEditorBlueprintLibrary::GetLoopMode()
+{
+	if (CurrentSequencer.IsValid())
+	{
+		return CurrentSequencer.Pin()->GetSequencerSettings()->GetLoopMode();
+	}
+
+	return ESequencerLoopMode::SLM_NoLoop;
+}
+
 void ULevelSequenceEditorBlueprintLibrary::PlayTo(FMovieSceneSequencePlaybackParams PlaybackParams)
 {
 	if (CurrentSequencer.IsValid())
diff --git a/Engine/Plugins/MovieScene/LevelSequenceEditor/Source/LevelSequenceEditor/Public/LevelSequenceEditorBlueprintLibrary.h b/Engine/Plugins/MovieScene/LevelSequenceEditor/Source/LevelSequenceEditor/Public/LevelSequenceEditorBlueprintLibrary.h
--- a/Engine/Plugins/MovieScene/LevelSequenceEditor/Source/LevelSequenceEditor/Public/LevelSequenceEditorBlueprintLibrary.h
+++ b/Engine/Plugins/MovieScene/LevelSequenceEditor/Source/LevelSequenceEditor/Public/LevelSequenceEditorBlueprintLibrary.h
@@ -4,6 +4,7 @@
 
 #include "Kismet/BlueprintFunctionLibrary.h"
 #include "SequencerCurveEditorObject.h"
+#include "SequencerSettings.h"
 #include "LevelSequenceEditorBlueprintLibrary.generated.h"
 
 struct FMovieSceneBindingProxy;
@@ -117,6 +118,18 @@ public:
 	UFUNCTION(BlueprintCallable, Category = "Level Sequence Editor")
 	static float GetPlaybackSpeed();
 
+	/**
+     * Set loop mode (note this is a per user preference)
+     */
+	UFUNCTION(BlueprintCallable, Category = "Level Sequence Editor")
+	static void SetLoopMode(ESequencerLoopMode NewLoopMode);
+
+	/**
+	 * Get loop mode (note this is a per user preference)
+	 */
+	UFUNCTION(BlueprintCallable, Category = "Level Sequence Editor")
+	static ESequencerLoopMode GetLoopMode();
+
 	/**
 	 * Play from the current time to the requested time in frames
 	 */
diff --git a/Engine/Source/Editor/Sequencer/Public/SequencerSettings.h b/Engine/Source/Editor/Sequencer/Public/SequencerSettings.h
--- a/Engine/Source/Editor/Sequencer/Public/SequencerSettings.h
+++ b/Engine/Source/Editor/Sequencer/Public/SequencerSettings.h
@@ -35,7 +35,7 @@ enum ESequencerZoomPosition : int
 	SZP_MousePosition UMETA(DisplayName="Mouse Position"),
 };
 
-UENUM()
+UENUM(BlueprintType)
 enum ESequencerLoopMode : int
 {
 	/** No Looping. */
