Sequencer: LevelSequenceEditorSubsystem now returns a SequencerModuleScriptingLayer

#rnx

diff --git a/Engine/Plugins/MovieScene/LevelSequenceEditor/Source/LevelSequenceEditor/Private/LevelSequenceEditorSubsystem.cpp b/Engine/Plugins/MovieScene/LevelSequenceEditor/Source/LevelSequenceEditor/Private/LevelSequenceEditorSubsystem.cpp
--- a/Engine/Plugins/MovieScene/LevelSequenceEditor/Source/LevelSequenceEditor/Private/LevelSequenceEditorSubsystem.cpp
+++ b/Engine/Plugins/MovieScene/LevelSequenceEditor/Source/LevelSequenceEditor/Private/LevelSequenceEditorSubsystem.cpp
@@ -2,7 +2,7 @@
 
 #include "LevelSequenceEditorSubsystem.h"
 #include "MVVM/ViewModels/SequencerEditorViewModel.h"
-#include "Scripting/SequencerScriptingLayer.h"
+#include "Scripting/SequencerModuleScriptingLayer.h"
 #include "SequencerCurveEditorObject.h"
 #include "Evaluation/MovieScenePlayback.h"
 #include "ISequencerModule.h"
@@ -249,12 +249,12 @@ TSharedPtr<ISequencer> ULevelSequenceEditorSubsystem::GetActiveSequencer()
 	return nullptr;
 }
 
-USequencerScriptingLayer* ULevelSequenceEditorSubsystem::GetScriptingLayer()
+USequencerModuleScriptingLayer* ULevelSequenceEditorSubsystem::GetScriptingLayer()
 {
 	TSharedPtr<ISequencer> Sequencer = GetActiveSequencer();
 	if (Sequencer)
 	{
-		return Sequencer->GetViewModel()->GetScriptingLayer();
+		return Cast<USequencerModuleScriptingLayer>(Sequencer->GetViewModel()->GetScriptingLayer());
 	}
 	return nullptr;
 }
diff --git a/Engine/Plugins/MovieScene/LevelSequenceEditor/Source/LevelSequenceEditor/Public/LevelSequenceEditorSubsystem.h b/Engine/Plugins/MovieScene/LevelSequenceEditor/Source/LevelSequenceEditor/Public/LevelSequenceEditorSubsystem.h
--- a/Engine/Plugins/MovieScene/LevelSequenceEditor/Source/LevelSequenceEditor/Public/LevelSequenceEditorSubsystem.h
+++ b/Engine/Plugins/MovieScene/LevelSequenceEditor/Source/LevelSequenceEditor/Public/LevelSequenceEditorSubsystem.h
@@ -32,7 +32,7 @@ class UMovieSceneCompiledDataManager;
 class UMovieSceneFolder;
 class UMovieSceneSection;
 class UMovieSceneSequence;
-class USequencerScriptingLayer;
+class USequencerModuleScriptingLayer;
 class IStructureDetailsView;
 class USequencerCurveEditorObject;
 
@@ -89,9 +89,9 @@ public:
 
 	void OnSequencerCreated(TSharedRef<ISequencer> InSequencer);
 
-	/** Retrieve the outliner */
+	/** Retrieve the scripting layer */
 	UFUNCTION(BlueprintPure, Category = "Level Sequence Editor")
-	USequencerScriptingLayer* GetScriptingLayer();
+	USequencerModuleScriptingLayer* GetScriptingLayer();
 
 	/** Retrieve the curve editor */
 	UFUNCTION(BlueprintPure, Category = "Level Sequence Editor")
diff --git a/Engine/Source/Editor/Sequencer/Private/Scripting/SequencerModuleScriptingLayer.cpp b/Engine/Source/Editor/Sequencer/Private/Scripting/SequencerModuleScriptingLayer.cpp
--- a/Engine/Source/Editor/Sequencer/Private/Scripting/SequencerModuleScriptingLayer.cpp
+++ b/Engine/Source/Editor/Sequencer/Private/Scripting/SequencerModuleScriptingLayer.cpp
@@ -10,4 +10,9 @@ void USequencerModuleScriptingLayer::Initialize(TSharedPtr<UE::Sequencer::FEdito
 {
 	Outliner = NewObject<USequencerModuleOutlinerScriptingObject>(this, "Outliner");
 	Outliner->Initialize(InViewModel->GetOutliner());
+}
+
+USequencerModuleOutlinerScriptingObject* USequencerModuleScriptingLayer::GetOutliner()
+{
+	return Cast<USequencerModuleOutlinerScriptingObject>(Outliner);
 }
\ No newline at end of file
diff --git a/Engine/Source/Editor/Sequencer/Private/Scripting/SequencerModuleScriptingLayer.h b/Engine/Source/Editor/Sequencer/Private/Scripting/SequencerModuleScriptingLayer.h
deleted file mode 100644
diff --git a/Engine/Source/Editor/Sequencer/Public/Scripting/SequencerModuleScriptingLayer.h b/Engine/Source/Editor/Sequencer/Public/Scripting/SequencerModuleScriptingLayer.h
new file mode 100644
--- /dev/null
+++ b/Engine/Source/Editor/Sequencer/Public/Scripting/SequencerModuleScriptingLayer.h
@@ -0,0 +1,22 @@
+// Copyright Epic Games, Inc. All Rights Reserved.
+
+#pragma once
+
+#include "CoreTypes.h"
+#include "Scripting/SequencerScriptingLayer.h"
+
+#include "SequencerModuleScriptingLayer.generated.h"
+
+UCLASS(MinimalAPI)
+class USequencerModuleScriptingLayer : public USequencerScriptingLayer
+{
+public:
+
+	GENERATED_BODY()
+
+	virtual void Initialize(TSharedPtr<UE::Sequencer::FEditorViewModel> InViewModel);
+
+	/** Retrieve the outliner */
+	UFUNCTION(BlueprintPure, Category = "Level Sequence Editor")
+	USequencerModuleOutlinerScriptingObject* GetOutliner();
+};
\ No newline at end of file
