Sequencer: Deprecated Get/SetCurrentTime and Get/SetCurrentLocalTime in favor of Get/SetGlobalPosition and Get/SetLocalPosition so that units can be specified in frames, seconds, marks or timecode.

Add timecode as a EUpdatePositionMethod

diff --git a/Engine/Plugins/MovieScene/LevelSequenceEditor/Source/LevelSequenceEditor/Private/LevelSequenceEditorBlueprintLibrary.cpp b/Engine/Plugins/MovieScene/LevelSequenceEditor/Source/LevelSequenceEditor/Private/LevelSequenceEditorBlueprintLibrary.cpp
--- a/Engine/Plugins/MovieScene/LevelSequenceEditor/Source/LevelSequenceEditor/Private/LevelSequenceEditorBlueprintLibrary.cpp
+++ b/Engine/Plugins/MovieScene/LevelSequenceEditor/Source/LevelSequenceEditor/Private/LevelSequenceEditorBlueprintLibrary.cpp
@@ -146,6 +146,32 @@ void ULevelSequenceEditorBlueprintLibrary::SetCurrentTime(int32 NewFrame)
 	}
 }
 
+void ULevelSequenceEditorBlueprintLibrary::SetGlobalPosition(FMovieSceneSequencePlaybackParams PlaybackParams, EMovieSceneTimeUnit TimeUnit)
+{
+	if (CurrentSequencer.IsValid())
+	{
+		FFrameRate DisplayRate = CurrentSequencer.Pin()->GetFocusedDisplayRate();
+		FFrameRate TickResolution = CurrentSequencer.Pin()->GetFocusedTickResolution();
+
+		UMovieSceneSequence* Sequence = CurrentSequencer.Pin()->GetFocusedMovieSceneSequence();
+
+		FFrameTime Position = PlaybackParams.GetPlaybackPosition(Sequence);
+		if (TimeUnit == EMovieSceneTimeUnit::DisplayRate)
+		{
+			Position = ConvertFrameTime(Position, DisplayRate, TickResolution);
+		}
+
+		if (Position == CurrentSequencer.Pin()->GetGlobalTime().Time)
+		{
+			CurrentSequencer.Pin()->ForceEvaluate();
+		}
+		else
+		{
+			CurrentSequencer.Pin()->SetGlobalTime(Position);
+		}
+	}
+}
+
 int32 ULevelSequenceEditorBlueprintLibrary::GetCurrentTime()
 {
 	if (CurrentSequencer.IsValid())
@@ -158,6 +184,25 @@ int32 ULevelSequenceEditorBlueprintLibrary::GetCurrentTime()
 	return 0;
 }
 
+FMovieSceneSequencePlaybackParams ULevelSequenceEditorBlueprintLibrary::GetGlobalPosition(EMovieSceneTimeUnit TimeUnit)
+{
+	FMovieSceneSequencePlaybackParams Params;
+	if (CurrentSequencer.IsValid())
+	{
+		FFrameRate DisplayRate = CurrentSequencer.Pin()->GetFocusedDisplayRate();
+		FFrameRate TickResolution = CurrentSequencer.Pin()->GetFocusedTickResolution();
+
+		FQualifiedFrameTime GlobalTime = CurrentSequencer.Pin()->GetGlobalTime();
+
+		Params.Frame = TimeUnit == EMovieSceneTimeUnit::DisplayRate ? ConvertFrameTime(GlobalTime.Time, TickResolution, DisplayRate) : GlobalTime.Time;
+		Params.Timecode = GlobalTime.ToTimecode();
+		Params.Time = GlobalTime.AsSeconds();
+		
+		return Params;
+	}
+	return Params;
+}
+
 void ULevelSequenceEditorBlueprintLibrary::SetCurrentLocalTime(int32 NewFrame)
 {
 	if (CurrentSequencer.IsValid())
@@ -177,6 +222,31 @@ void ULevelSequenceEditorBlueprintLibrary::SetCurrentLocalTime(int32 NewFrame)
 	}
 }
 
+void ULevelSequenceEditorBlueprintLibrary::SetLocalPosition(FMovieSceneSequencePlaybackParams PlaybackParams, EMovieSceneTimeUnit TimeUnit)
+{
+	if (CurrentSequencer.IsValid())
+	{
+		FFrameRate DisplayRate = CurrentSequencer.Pin()->GetFocusedDisplayRate();
+		FFrameRate TickResolution = CurrentSequencer.Pin()->GetFocusedTickResolution();
+
+		UMovieSceneSequence* Sequence = CurrentSequencer.Pin()->GetFocusedMovieSceneSequence();
+		FFrameTime Position = PlaybackParams.GetPlaybackPosition(Sequence);
+		if (TimeUnit == EMovieSceneTimeUnit::DisplayRate)
+		{
+			Position = ConvertFrameTime(Position, DisplayRate, TickResolution);
+		}
+
+		if (Position == CurrentSequencer.Pin()->GetLocalTime().Time)
+		{
+			CurrentSequencer.Pin()->ForceEvaluate();
+		}
+		else
+		{
+			CurrentSequencer.Pin()->SetLocalTime(Position);
+		}
+	}
+}
+
 int32 ULevelSequenceEditorBlueprintLibrary::GetCurrentLocalTime()
 {
 	if (CurrentSequencer.IsValid())
@@ -189,6 +259,25 @@ int32 ULevelSequenceEditorBlueprintLibrary::GetCurrentLocalTime()
 	return 0;
 }
 
+FMovieSceneSequencePlaybackParams ULevelSequenceEditorBlueprintLibrary::GetLocalPosition(EMovieSceneTimeUnit TimeUnit)
+{
+	FMovieSceneSequencePlaybackParams Params;
+	if (CurrentSequencer.IsValid())
+	{
+		FFrameRate DisplayRate = CurrentSequencer.Pin()->GetFocusedDisplayRate();
+		FFrameRate TickResolution = CurrentSequencer.Pin()->GetFocusedTickResolution();
+
+		FQualifiedFrameTime LocalTime = CurrentSequencer.Pin()->GetLocalTime();
+
+		Params.Frame = TimeUnit == EMovieSceneTimeUnit::DisplayRate ? ConvertFrameTime(LocalTime.Time, TickResolution, DisplayRate) : LocalTime.Time;
+		Params.Timecode = LocalTime.ToTimecode();
+		Params.Time = LocalTime.AsSeconds();
+
+		return Params;
+	}
+	return Params;
+}
+
 void ULevelSequenceEditorBlueprintLibrary::SetPlaybackSpeed(float NewPlaybackSpeed)
 {
 	if (CurrentSequencer.IsValid())
@@ -225,10 +314,18 @@ ESequencerLoopMode ULevelSequenceEditorBlueprintLibrary::GetLoopMode()
 	return ESequencerLoopMode::SLM_NoLoop;
 }
 
-void ULevelSequenceEditorBlueprintLibrary::PlayTo(FMovieSceneSequencePlaybackParams PlaybackParams)
+void ULevelSequenceEditorBlueprintLibrary::PlayTo(FMovieSceneSequencePlaybackParams PlaybackParams, EMovieSceneTimeUnit TimeUnit)
 {
 	if (CurrentSequencer.IsValid())
 	{
+		FFrameRate DisplayRate = CurrentSequencer.Pin()->GetFocusedDisplayRate();
+		FFrameRate TickResolution = CurrentSequencer.Pin()->GetFocusedTickResolution();
+
+		if (TimeUnit == EMovieSceneTimeUnit::DisplayRate)
+		{
+			PlaybackParams.Frame = ConvertFrameTime(PlaybackParams.Frame, DisplayRate, TickResolution);
+		}
+
 		CurrentSequencer.Pin()->PlayTo(PlaybackParams);
 	}
 }
diff --git a/Engine/Plugins/MovieScene/LevelSequenceEditor/Source/LevelSequenceEditor/Public/LevelSequenceEditorBlueprintLibrary.h b/Engine/Plugins/MovieScene/LevelSequenceEditor/Source/LevelSequenceEditor/Public/LevelSequenceEditorBlueprintLibrary.h
--- a/Engine/Plugins/MovieScene/LevelSequenceEditor/Source/LevelSequenceEditor/Public/LevelSequenceEditorBlueprintLibrary.h
+++ b/Engine/Plugins/MovieScene/LevelSequenceEditor/Source/LevelSequenceEditor/Public/LevelSequenceEditorBlueprintLibrary.h
@@ -3,6 +3,7 @@
 #pragma once
 
 #include "Kismet/BlueprintFunctionLibrary.h"
+#include "MovieSceneTimeUnit.h"
 #include "SequencerCurveEditorObject.h"
 #include "SequencerSettings.h"
 #include "LevelSequenceEditorBlueprintLibrary.generated.h"
@@ -82,29 +83,45 @@ public:
 
 public:
 
+	UE_DEPRECATED(5.4, "Use SetCurrentTime that takes a FMovieSceneSequencePlaybackParams")
+	UFUNCTION(BlueprintCallable, Category = "Level Sequence Editor", meta = (DeprecatedFunction, DeprecationMessage = "Use SetCurrentTime that takes a FMovieSceneSequencePlaybackParams"))
+	static void SetCurrentTime(int32 NewFrame);
+	
+	UE_DEPRECATED(5.4, "Use GetCurrentTime that returns a FMovieSceneSequencePlaybackParams")
+	UFUNCTION(BlueprintCallable, Category = "Level Sequence Editor", meta = (DeprecatedFunction, DeprecationMessage = "Use GetCurrentTime that returns a FMovieSceneSequencePlaybackParams"))
+	static int32 GetCurrentTime();
+
+	UE_DEPRECATED(5.4, "Use SetCurrentLocalTime that takes a FMovieSceneSequencePlaybackParams")
+	UFUNCTION(BlueprintCallable, Category = "Level Sequence Editor", meta = (DeprecatedFunction, DeprecationMessage = "Use SetCurrentLocalTime that takes a FMovieSceneSequencePlaybackParams"))
+	static void SetCurrentLocalTime(int32 NewFrame);
+
+	UE_DEPRECATED(5.4, "Use GetCurrentLocalTime that takes a FMovieSceneSequencePlaybackParams")
+	UFUNCTION(BlueprintCallable, Category = "Level Sequence Editor", meta = (DeprecatedFunction, DeprecationMessage = "Use GetCurrentLocalTime that returns a FMovieSceneSequencePlaybackParams"))
+	static int32 GetCurrentLocalTime();
+
 	/**
-	 * Set global playback position for the current level sequence in frames. If the requested time is the same as the current time, an evaluation will be forced.
+	 * Set global playhead position for the current level sequence. If the requested time is the same as the current time, an evaluation will be forced.
 	 */
-	UFUNCTION(BlueprintCallable, Category = "Level Sequence Editor")
-	static void SetCurrentTime(int32 NewFrame);
+	UFUNCTION(BlueprintCallable, Category = "Level Sequence Editor", DisplayName = "Set Current Time")
+	static void SetGlobalPosition(FMovieSceneSequencePlaybackParams PlaybackParams, EMovieSceneTimeUnit TimeUnit = EMovieSceneTimeUnit::DisplayRate);
 
 	/**
-	 * Get the current global playback position in frames
+	 * Get the current global playhead position
 	 */
-	UFUNCTION(BlueprintCallable, Category = "Level Sequence Editor")
-	static int32 GetCurrentTime();
+	UFUNCTION(BlueprintCallable, Category = "Level Sequence Editor", DisplayName = "Get Current Time")
+	static FMovieSceneSequencePlaybackParams GetGlobalPosition(EMovieSceneTimeUnit TimeUnit = EMovieSceneTimeUnit::DisplayRate);
 
 	/**
-	 * Set local playback position for the current level sequence in frames. If the requested time is the same as the current time, an evaluation will be forced.
+	 * Set local playhead position for the current level sequence. If the requested time is the same as the current time, an evaluation will be forced.
 	 */
-	UFUNCTION(BlueprintCallable, Category = "Level Sequence Editor")
-	static void SetCurrentLocalTime(int32 NewFrame);
+	UFUNCTION(BlueprintCallable, Category = "Level Sequence Editor", DisplayName = "Set Current Local Time")
+	static void SetLocalPosition(FMovieSceneSequencePlaybackParams PlaybackParams, EMovieSceneTimeUnit TimeUnit = EMovieSceneTimeUnit::DisplayRate);
 
 	/**
-	 * Get the current local playback position in frames
+	 * Get the current local playhead position
 	 */
-	UFUNCTION(BlueprintCallable, Category = "Level Sequence Editor")
-	static int32 GetCurrentLocalTime();
+	UFUNCTION(BlueprintCallable, Category = "Level Sequence Editor", DisplayName = "Get Current Local Time")
+	static FMovieSceneSequencePlaybackParams GetLocalPosition(EMovieSceneTimeUnit TimeUnit = EMovieSceneTimeUnit::DisplayRate);
 
 	/**
 	 * Set playback speed of the current level sequence
@@ -134,7 +151,7 @@ public:
 	 * Play from the current time to the requested time in frames
 	 */
 	UFUNCTION(BlueprintCallable, Category = "Level Sequence Editor")
-	static void PlayTo(FMovieSceneSequencePlaybackParams PlaybackParams);
+	static void PlayTo(FMovieSceneSequencePlaybackParams PlaybackParams, EMovieSceneTimeUnit TimeUnit = EMovieSceneTimeUnit::DisplayRate);
 
 public:
 
diff --git a/Engine/Plugins/VirtualProduction/VirtualCamera/Source/VirtualCamera/Private/FunctionLibraries/VCamBlueprintFunctionLibrary.cpp b/Engine/Plugins/VirtualProduction/VirtualCamera/Source/VirtualCamera/Private/FunctionLibraries/VCamBlueprintFunctionLibrary.cpp
--- a/Engine/Plugins/VirtualProduction/VirtualCamera/Source/VirtualCamera/Private/FunctionLibraries/VCamBlueprintFunctionLibrary.cpp
+++ b/Engine/Plugins/VirtualProduction/VirtualCamera/Source/VirtualCamera/Private/FunctionLibraries/VCamBlueprintFunctionLibrary.cpp
@@ -14,6 +14,7 @@
 #include "GameFramework/PlayerController.h"
 #include "LevelSequence.h"
 #include "MovieScene.h"
+#include "MovieSceneSequencePlayer.h"
 #include "Slate/SceneViewport.h"
 
 #if WITH_EDITOR
@@ -94,7 +95,9 @@ void UVCamBlueprintFunctionLibrary::PauseCurrentLevelSequence()
 void UVCamBlueprintFunctionLibrary::SetCurrentLevelSequenceCurrentFrame(int32 NewFrame)
 {
 #if WITH_EDITOR
-	ULevelSequenceEditorBlueprintLibrary::SetCurrentTime(NewFrame);
+	FMovieSceneSequencePlaybackParams Params(FFrameTime(NewFrame), EUpdatePositionMethod::Play);
+	
+	ULevelSequenceEditorBlueprintLibrary::SetGlobalPosition(Params);
 #endif
 
 }
@@ -102,7 +105,8 @@ void UVCamBlueprintFunctionLibrary::SetCurrentLevelSequenceCurrentFrame(int32 Ne
 int32 UVCamBlueprintFunctionLibrary::GetCurrentLevelSequenceCurrentFrame()
 {
 #if WITH_EDITOR
-	return ULevelSequenceEditorBlueprintLibrary::GetCurrentTime();
+	FMovieSceneSequencePlaybackParams Position = ULevelSequenceEditorBlueprintLibrary::GetGlobalPosition();
+	return Position.Frame.FloorToFrame().Value;
 #else
 	return 0;
 #endif
diff --git a/Engine/Source/Runtime/MovieScene/Private/MovieSceneSequencePlayer.cpp b/Engine/Source/Runtime/MovieScene/Private/MovieSceneSequencePlayer.cpp
--- a/Engine/Source/Runtime/MovieScene/Private/MovieSceneSequencePlayer.cpp
+++ b/Engine/Source/Runtime/MovieScene/Private/MovieSceneSequencePlayer.cpp
@@ -116,6 +116,48 @@ FFrameTime FMovieSceneSequencePlaybackParams::GetPlaybackPosition(UMovieSceneSeq
 			}
 		}
 	}
+	else if (PositionType == EMovieScenePositionType::Timecode)
+	{
+		PlaybackPosition = Timecode.ToFrameNumber(Player->GetFrameRate());
+	}
+
+	return PlaybackPosition;
+}
+
+FFrameTime FMovieSceneSequencePlaybackParams::GetPlaybackPosition(UMovieSceneSequence* Sequence) const
+{
+	FFrameTime PlaybackPosition;
+
+	UMovieScene* MovieScene = Sequence->GetMovieScene();
+	if (!MovieScene)
+	{
+		return PlaybackPosition;
+	}
+
+	FFrameRate DisplayRate = MovieScene->GetDisplayRate();
+	FFrameRate TickResolution = MovieScene->GetTickResolution();
+
+	if (PositionType == EMovieScenePositionType::Frame)
+	{
+		PlaybackPosition = Frame;
+	}
+	else if (PositionType == EMovieScenePositionType::Time)
+	{
+		PlaybackPosition = Time * DisplayRate;
+	}
+	else if (PositionType == EMovieScenePositionType::MarkedFrame)
+	{
+		int32 MarkedIndex = MovieScene->FindMarkedFrameByLabel(MarkedFrame);
+		
+		if (MarkedIndex != INDEX_NONE)
+		{
+			PlaybackPosition = ConvertFrameTime(MovieScene->GetMarkedFrames()[MarkedIndex].FrameNumber, TickResolution, DisplayRate);
+		}
+	}
+	else if (PositionType == EMovieScenePositionType::Timecode)
+	{
+		PlaybackPosition = Timecode.ToFrameNumber(DisplayRate);
+	}
 
 	return PlaybackPosition;
 }
diff --git a/Engine/Source/Runtime/MovieScene/Public/MovieSceneSequencePlayer.h b/Engine/Source/Runtime/MovieScene/Public/MovieSceneSequencePlayer.h
--- a/Engine/Source/Runtime/MovieScene/Public/MovieSceneSequencePlayer.h
+++ b/Engine/Source/Runtime/MovieScene/Public/MovieSceneSequencePlayer.h
@@ -85,6 +85,7 @@ enum class EMovieScenePositionType : uint8
 	Frame,
 	Time,
 	MarkedFrame,
+	Timecode
 };
 
 USTRUCT(BlueprintType)
@@ -122,7 +123,19 @@ struct FMovieSceneSequencePlaybackParams
 		, bHasJumped(false)
 	{}
 
-	FFrameTime GetPlaybackPosition(UMovieSceneSequencePlayer* Player) const;
+	FMovieSceneSequencePlaybackParams(const FTimecode& InTimecode, EUpdatePositionMethod InUpdateMethod)
+		: Time(0.f)
+		, Timecode(InTimecode)
+		, PositionType(EMovieScenePositionType::Timecode)
+		, UpdateMethod(InUpdateMethod)
+		, bHasJumped(false)
+	{}
+
+	// Get the playback position using the player's tick resolution and display rate	
+	MOVIESCENE_API FFrameTime GetPlaybackPosition(UMovieSceneSequencePlayer* Player) const;
+
+	// Get the playback position using the sequence's tick resolution and display rate
+	MOVIESCENE_API FFrameTime GetPlaybackPosition(UMovieSceneSequence* Sequence) const;
 
 	UPROPERTY(EditAnywhere, BlueprintReadWrite, Category="Cinematic", meta=(EditCondition="PositionType == EMovieScenePositionType::Frame"))
 	FFrameTime Frame;
@@ -133,6 +146,9 @@ struct FMovieSceneSequencePlaybackParams
 	UPROPERTY(EditAnywhere, BlueprintReadWrite, Category="Cinematic", meta=(EditCondition="PositionType == EMovieScenePositionType::MarkedFrame"))
 	FString MarkedFrame;
 
+	UPROPERTY(EditAnywhere, BlueprintReadWrite, Category = "Cinematic", meta=(EditCondition="PositionType == EMovieScenePositionType::Timecode"))
+	FTimecode Timecode;
+
 	UPROPERTY(EditAnywhere, BlueprintReadWrite, Category="Cinematic")
 	EMovieScenePositionType PositionType;
 
