Sequencer- disable manual editing of locator flags as it's currently possible to create combinations that cause crashes and other issues.

diff --git a/Engine/Plugins/MovieScene/LevelSequenceEditor/Source/LevelSequenceEditor/Public/LevelSequenceEditorSubsystem.h b/Engine/Plugins/MovieScene/LevelSequenceEditor/Source/LevelSequenceEditor/Public/LevelSequenceEditorSubsystem.h
--- a/Engine/Plugins/MovieScene/LevelSequenceEditor/Source/LevelSequenceEditor/Public/LevelSequenceEditorSubsystem.h
+++ b/Engine/Plugins/MovieScene/LevelSequenceEditor/Source/LevelSequenceEditor/Public/LevelSequenceEditorSubsystem.h
@@ -58,7 +58,7 @@ struct FMovieSceneUniversalLocatorInfo
 	FUniversalObjectLocator Locator;
 
 	// Flags for how to resolve the locator
-	UPROPERTY(EditAnywhere, Category = "Default")
+	UPROPERTY()
 	ELocatorResolveFlags ResolveFlags = ELocatorResolveFlags::None;
 };
 
