Fix styles not being found when using LiveLinkHub
Using the plugin content path from the plugin manager allows the hub to use resources located in a different directory than /engine/content

diff --git a/Engine/Plugins/MovieScene/LevelSequenceEditor/Source/LevelSequenceEditor/LevelSequenceEditor.Build.cs b/Engine/Plugins/MovieScene/LevelSequenceEditor/Source/LevelSequenceEditor/LevelSequenceEditor.Build.cs
--- a/Engine/Plugins/MovieScene/LevelSequenceEditor/Source/LevelSequenceEditor/LevelSequenceEditor.Build.cs
+++ b/Engine/Plugins/MovieScene/LevelSequenceEditor/Source/LevelSequenceEditor/LevelSequenceEditor.Build.cs
@@ -47,8 +47,9 @@ public class LevelSequenceEditor : ModuleRules
 				"MovieScene",
                 "MovieSceneTools",
 				"MovieSceneTracks",
+				"Projects",
                 "PropertyEditor",
-                "SequencerCore",
+				"SequencerCore",
 				"Sequencer",
                 "Slate",
                 "SlateCore",
diff --git a/Engine/Plugins/MovieScene/LevelSequenceEditor/Source/LevelSequenceEditor/Private/Styles/LevelSequenceEditorStyle.h b/Engine/Plugins/MovieScene/LevelSequenceEditor/Source/LevelSequenceEditor/Private/Styles/LevelSequenceEditorStyle.h
--- a/Engine/Plugins/MovieScene/LevelSequenceEditor/Source/LevelSequenceEditor/Private/Styles/LevelSequenceEditorStyle.h
+++ b/Engine/Plugins/MovieScene/LevelSequenceEditor/Source/LevelSequenceEditor/Private/Styles/LevelSequenceEditorStyle.h
@@ -5,6 +5,7 @@
 #include "Misc/Paths.h"
 #include "Brushes/SlateBorderBrush.h"
 #include "Brushes/SlateImageBrush.h"
+#include "Interfaces/IPluginManager.h"
 #include "Styling/SlateStyleRegistry.h"
 #include "Styling/SlateStyle.h"
 #include "Styling/SlateStyleMacros.h"
@@ -26,7 +27,9 @@ public:
 		const FVector2D Icon40x40(40.0f, 40.0f);
 		const FVector2D Icon48x48(48.0f, 48.0f);
 
-		SetContentRoot(FPaths::EnginePluginsDir() / TEXT("MovieScene/LevelSequenceEditor/Content"));
+		const FString ContentDir = IPluginManager::Get().FindPlugin(TEXT("LevelSequenceEditor"))->GetContentDir();
+
+		SetContentRoot(ContentDir);
 		SetCoreContentRoot(FPaths::EngineContentDir() / TEXT("Slate"));
 
 		// tab icons
diff --git a/Engine/Plugins/VirtualProduction/Takes/Source/TakeRecorder/Private/TakeRecorderStyle.cpp b/Engine/Plugins/VirtualProduction/Takes/Source/TakeRecorder/Private/TakeRecorderStyle.cpp
--- a/Engine/Plugins/VirtualProduction/Takes/Source/TakeRecorder/Private/TakeRecorderStyle.cpp
+++ b/Engine/Plugins/VirtualProduction/Takes/Source/TakeRecorder/Private/TakeRecorderStyle.cpp
@@ -2,6 +2,7 @@
 
 #include "TakeRecorderStyle.h"
 
+#include "Interfaces/IPluginManager.h"
 #include "Styling/SlateTypes.h"
 #include "Styling/CoreStyle.h"
 #include "Styling/AppStyle.h"
@@ -27,7 +28,9 @@ FTakeRecorderStyle::FTakeRecorderStyle()
 	const FLinearColor DarkGrey( FColor(30, 30, 30) );
 	const FLinearColor Black(FLinearColor::Black);
 
-	SetContentRoot(FPaths::EnginePluginsDir() / TEXT("VirtualProduction/Takes/Content"));
+
+	const FString ContentDir = IPluginManager::Get().FindPlugin(TEXT("Takes"))->GetContentDir();
+	SetContentRoot(ContentDir);
 	SetCoreContentRoot(FPaths::EngineContentDir() / TEXT("Slate"));
 
 	FButtonStyle Button = FButtonStyle()
diff --git a/Engine/Plugins/VirtualProduction/Takes/Source/TakeRecorder/TakeRecorder.Build.cs b/Engine/Plugins/VirtualProduction/Takes/Source/TakeRecorder/TakeRecorder.Build.cs
--- a/Engine/Plugins/VirtualProduction/Takes/Source/TakeRecorder/TakeRecorder.Build.cs
+++ b/Engine/Plugins/VirtualProduction/Takes/Source/TakeRecorder/TakeRecorder.Build.cs
@@ -22,6 +22,7 @@ public class TakeRecorder : ModuleRules
 				"LevelSequence",
 				"MovieScene",
 				"MovieSceneTools",
+				"Projects",
 				"PropertyEditor",
 				"TakesCore",
 				"TakeMovieScene",
