Sequencer: Disable saving spawnable state during GC.

This was causing the spawnable to be saved and the level sequence to be inadvertently marked dirty through OnObjectModified when components are removed from an actor during GC.

Tested that the spawnable state is still saved when you make changes to the spawnable and close Sequencer.

diff --git a/Engine/Plugins/MovieScene/LevelSequenceEditor/Source/LevelSequenceEditor/Private/Misc/LevelSequenceEditorSpawnRegister.cpp b/Engine/Plugins/MovieScene/LevelSequenceEditor/Source/LevelSequenceEditor/Private/Misc/LevelSequenceEditorSpawnRegister.cpp
--- a/Engine/Plugins/MovieScene/LevelSequenceEditor/Source/LevelSequenceEditor/Private/Misc/LevelSequenceEditorSpawnRegister.cpp
+++ b/Engine/Plugins/MovieScene/LevelSequenceEditor/Source/LevelSequenceEditor/Private/Misc/LevelSequenceEditorSpawnRegister.cpp
@@ -21,6 +21,7 @@
 FLevelSequenceEditorSpawnRegister::FLevelSequenceEditorSpawnRegister()
 {
 	bShouldClearSelectionCache = true;
+	bIsEngineCollectingGarbage = false;
 
 	FLevelEditorModule& LevelEditor = FModuleManager::GetModuleChecked<FLevelEditorModule>("LevelEditor");
 	OnActorSelectionChangedHandle = LevelEditor.OnActorSelectionChanged().AddRaw(this, &FLevelSequenceEditorSpawnRegister::HandleActorSelectionChanged);
@@ -31,6 +32,9 @@ FLevelSequenceEditorSpawnRegister::FLevelSequenceEditorSpawnRegister()
 	OnObjectModifiedHandle = FCoreUObjectDelegates::OnObjectModified.AddRaw(this, &FLevelSequenceEditorSpawnRegister::OnObjectModified);
 	OnObjectSavedHandle    = FCoreUObjectDelegates::OnObjectPreSave.AddRaw(this, &FLevelSequenceEditorSpawnRegister::OnPreObjectSaved);
 #endif
+
+	OnPreGarbageCollectHandle = FCoreUObjectDelegates::GetPreGarbageCollectDelegate().AddRaw(this, &FLevelSequenceEditorSpawnRegister::UpdateIsEngineCollectingGarbage, true);
+	OnPostGarbageCollectHandle = FCoreUObjectDelegates::GetPostGarbageCollect().AddRaw(this, &FLevelSequenceEditorSpawnRegister::UpdateIsEngineCollectingGarbage, false);
 }
 
 
@@ -53,6 +57,9 @@ FLevelSequenceEditorSpawnRegister::~FLevelSequenceEditorSpawnRegister()
 	FCoreUObjectDelegates::OnObjectModified.Remove(OnObjectModifiedHandle);
 	FCoreUObjectDelegates::OnObjectPreSave.Remove(OnObjectSavedHandle);
 #endif
+
+	FCoreUObjectDelegates::GetPostGarbageCollect().Remove(OnPostGarbageCollectHandle);
+	FCoreUObjectDelegates::GetPreGarbageCollectDelegate().Remove(OnPreGarbageCollectHandle);
 }
 
 
@@ -256,6 +263,13 @@ void FLevelSequenceEditorSpawnRegister::OnObjectsReplaced(const TMap<UObject*, U
 
 void FLevelSequenceEditorSpawnRegister::OnObjectModified(UObject* ModifiedObject)
 {
+	// If we are reinstancing then renaming existing objects aside can cause this callback to be called,
+	// which can end up with unnecessary template objects getting updated when we close sequencer down
+	if (GIsReinstancing || GIsGCingAfterBlueprintCompile || bIsEngineCollectingGarbage)
+	{
+		return;
+	}
+
 	// If the sequence is evaluating, we don't want object modifications to dirty the sequence itself. 
 	// For example, this protects against situations where OnObjectModified would be called in response 
 	// to the spawnable being attached with AttachToComponent
@@ -391,5 +405,9 @@ bool FLevelSequenceEditorSpawnRegister::CanConvertSpawnableToPossessable(FMovieS
 
 #endif
 
+void FLevelSequenceEditorSpawnRegister::UpdateIsEngineCollectingGarbage(bool bIsCollectingGarbage)
+{
+	bIsEngineCollectingGarbage = bIsCollectingGarbage;
+}
 
 #undef LOCTEXT_NAMESPACE
diff --git a/Engine/Plugins/MovieScene/LevelSequenceEditor/Source/LevelSequenceEditor/Public/Misc/LevelSequenceEditorSpawnRegister.h b/Engine/Plugins/MovieScene/LevelSequenceEditor/Source/LevelSequenceEditor/Public/Misc/LevelSequenceEditorSpawnRegister.h
--- a/Engine/Plugins/MovieScene/LevelSequenceEditor/Source/LevelSequenceEditor/Public/Misc/LevelSequenceEditorSpawnRegister.h
+++ b/Engine/Plugins/MovieScene/LevelSequenceEditor/Source/LevelSequenceEditor/Public/Misc/LevelSequenceEditorSpawnRegister.h
@@ -63,6 +63,9 @@ private:
 	/** Called before an object is saved in the editor */
 	void OnPreObjectSaved(UObject* Object, FObjectPreSaveContext SaveContext);
 
+	/** Called on pre/post GC */
+	void UpdateIsEngineCollectingGarbage(bool bIsCollectingGarbage);
+
 private:
 
 	struct FTrackedObjectState
@@ -104,6 +107,13 @@ private:
 
 	/** Handle to a delegate that is bound to FCoreUObjectDelegates::OnObjectPreSave to harvest changes to spawned objects. */
 	FDelegateHandle OnObjectSavedHandle;
+
+	/** True when GCing */
+	bool bIsEngineCollectingGarbage;
+
+	/** Handle to a delegate that is bound to FCoreUObjectDelegates::OnPre/PostGarbageCollectHandle to disable saving changes while GCing. */
+	FDelegateHandle OnPreGarbageCollectHandle;
+	FDelegateHandle OnPostGarbageCollectHandle;
 };
 
 #if UE_ENABLE_INCLUDE_ORDER_DEPRECATED_IN_5_2
