Sequencer- fix to issue with binding properties menu where making a change could lose data.

diff --git a/Engine/Plugins/MovieScene/LevelSequenceEditor/Source/LevelSequenceEditor/Private/LevelSequenceEditorSubsystem.cpp b/Engine/Plugins/MovieScene/LevelSequenceEditor/Source/LevelSequenceEditor/Private/LevelSequenceEditorSubsystem.cpp
--- a/Engine/Plugins/MovieScene/LevelSequenceEditor/Source/LevelSequenceEditor/Private/LevelSequenceEditorSubsystem.cpp
+++ b/Engine/Plugins/MovieScene/LevelSequenceEditor/Source/LevelSequenceEditor/Private/LevelSequenceEditorSubsystem.cpp
@@ -1971,24 +1971,12 @@ void ULevelSequenceEditorSubsystem::OnFinishedChangingLocators(const FPropertyCh
 		// Add the new updated bindings
 		for (FMovieSceneBindingPropertyInfo& LocatorInfo : BindingPropertyInfoList->Bindings)
 		{
-			if (PropertyChangedEvent.Property != nullptr && PropertyChangedEvent.Property->GetFName() == TEXT("Locator"))
+			UMovieSceneCustomBinding* CopiedBinding = nullptr;
+			if (LocatorInfo.CustomBinding)
 			{
-				// Ensure flags are initialized from scratch
-				const FMovieSceneBindingReference* NewRef = BindingReferences->AddBinding(ObjectBindingID, MoveTemp(LocatorInfo.Locator));
-				if (NewRef)
-				{
-					LocatorInfo.ResolveFlags = NewRef->ResolveFlags;
-				}
-			}
-			else
-			{
-				UMovieSceneCustomBinding* CopiedBinding = nullptr;
-				if (LocatorInfo.CustomBinding)
-				{
-					CopiedBinding = Cast<UMovieSceneCustomBinding>(StaticDuplicateObject(LocatorInfo.CustomBinding, MovieScene));
-				}
-				BindingReferences->AddBinding(ObjectBindingID, MoveTemp(LocatorInfo.Locator), LocatorInfo.ResolveFlags, CopiedBinding);
+				CopiedBinding = Cast<UMovieSceneCustomBinding>(StaticDuplicateObject(LocatorInfo.CustomBinding, MovieScene));
 			}
+			BindingReferences->AddBinding(ObjectBindingID, MoveTemp(LocatorInfo.Locator), LocatorInfo.ResolveFlags, CopiedBinding);
 		}
 
 		Sequencer->State.Invalidate(ObjectBindingID, Sequencer->GetFocusedTemplateID());
