Sequencer- fix out of bounds array crash in binding customization code.

diff --git a/Engine/Plugins/MovieScene/LevelSequenceEditor/Source/LevelSequenceEditor/Private/Misc/LevelSequenceCustomization.cpp b/Engine/Plugins/MovieScene/LevelSequenceEditor/Source/LevelSequenceEditor/Private/Misc/LevelSequenceCustomization.cpp
--- a/Engine/Plugins/MovieScene/LevelSequenceEditor/Source/LevelSequenceEditor/Private/Misc/LevelSequenceCustomization.cpp
+++ b/Engine/Plugins/MovieScene/LevelSequenceEditor/Source/LevelSequenceEditor/Private/Misc/LevelSequenceCustomization.cpp
@@ -263,25 +263,29 @@ void FLevelSequenceCustomization::ExtendObjectBindingContextMenu(FMenuBuilder& M
 						FNewMenuDelegate::CreateRaw(this, &FLevelSequenceCustomization::AddDynamicPossessionMenu, ObjectBindingModel));
 				}
 				// Custom bindings
-				const FMovieSceneBindingReference& CurrentBindingReference = Sequencer->GetFocusedMovieSceneSequence()->GetBindingReferences()->GetReferences(ObjectBindingID)[0];
-				TArrayView<TWeakObjectPtr<>> BoundObjects = Sequencer->FindBoundObjects(ObjectBindingID, Sequencer->GetFocusedTemplateID());
-				if (BoundObjects.Num() == 1)
+				TArrayView<const FMovieSceneBindingReference> BindingReferencesList = Sequencer->GetFocusedMovieSceneSequence()->GetBindingReferences()->GetReferences(ObjectBindingID);
+				if (BindingReferencesList.Num() == 1)
 				{
-					if (UObject* CurrentBoundObject = BoundObjects[0].Get())
+					const FMovieSceneBindingReference& CurrentBindingReference = BindingReferencesList[0];
+					TArrayView<TWeakObjectPtr<>> BoundObjects = Sequencer->FindBoundObjects(ObjectBindingID, Sequencer->GetFocusedTemplateID());
+					if (BoundObjects.Num() == 1)
 					{
-						TArrayView<const TSubclassOf<UMovieSceneCustomBinding>> PrioritySortedCustomBindingTypes = Sequencer->GetSupportedCustomBindingTypes();
-						for (const TSubclassOf<UMovieSceneCustomBinding>& CustomBindingType : PrioritySortedCustomBindingTypes)
+						if (UObject* CurrentBoundObject = BoundObjects[0].Get())
 						{
-							if (CustomBindingType
-								&& (!CurrentBindingReference.CustomBinding
-									|| CurrentBindingReference.CustomBinding->GetClass() != CustomBindingType)
-								&& CustomBindingType->GetDefaultObject<UMovieSceneCustomBinding>()->SupportsConversionFromBinding(CurrentBindingReference, CurrentBoundObject))
+							TArrayView<const TSubclassOf<UMovieSceneCustomBinding>> PrioritySortedCustomBindingTypes = Sequencer->GetSupportedCustomBindingTypes();
+							for (const TSubclassOf<UMovieSceneCustomBinding>& CustomBindingType : PrioritySortedCustomBindingTypes)
 							{
-								MenuBuilder.AddMenuEntry(
-									FText::Format(LOCTEXT("ConvertToCustomBinding", "Convert to {0}"), CustomBindingType->GetDefaultObject<UMovieSceneCustomBinding>()->GetBindingTypePrettyName()),
-									FText::Format(LOCTEXT("ConvertToCustomBindingTooltip", "Convert selected binding to {0}"), CustomBindingType->GetDefaultObject<UMovieSceneCustomBinding>()->GetBindingTypePrettyName()),
-									FSlateIcon(),
-									FUIAction(FExecuteAction::CreateRaw(this, &FLevelSequenceCustomization::ConvertToCustomBinding, ObjectBindingModel, CustomBindingType)));
+								if (CustomBindingType
+									&& (!CurrentBindingReference.CustomBinding
+										|| CurrentBindingReference.CustomBinding->GetClass() != CustomBindingType)
+									&& CustomBindingType->GetDefaultObject<UMovieSceneCustomBinding>()->SupportsConversionFromBinding(CurrentBindingReference, CurrentBoundObject))
+								{
+									MenuBuilder.AddMenuEntry(
+										FText::Format(LOCTEXT("ConvertToCustomBinding", "Convert to {0}"), CustomBindingType->GetDefaultObject<UMovieSceneCustomBinding>()->GetBindingTypePrettyName()),
+										FText::Format(LOCTEXT("ConvertToCustomBindingTooltip", "Convert selected binding to {0}"), CustomBindingType->GetDefaultObject<UMovieSceneCustomBinding>()->GetBindingTypePrettyName()),
+										FSlateIcon(),
+										FUIAction(FExecuteAction::CreateRaw(this, &FLevelSequenceCustomization::ConvertToCustomBinding, ObjectBindingModel, CustomBindingType)));
+								}
 							}
 						}
 					}
@@ -310,23 +314,27 @@ void FLevelSequenceCustomization::ExtendObjectBindingContextMenu(FMenuBuilder& M
 					MenuBuilder.AddMenuEntry(FSequencerCommands::Get().ConvertToPossessable);
 
 					// Custom bindings
-					const FMovieSceneBindingReference& CurrentBindingReference = Sequencer->GetFocusedMovieSceneSequence()->GetBindingReferences()->GetReferences(ObjectBindingID)[0];
-					UObject* CurrentBoundObject = Sequencer->FindSpawnedObjectOrTemplate(ObjectBindingID);
-					if (CurrentBoundObject)
+					TArrayView<const FMovieSceneBindingReference> BindingReferencesList = Sequencer->GetFocusedMovieSceneSequence()->GetBindingReferences()->GetReferences(ObjectBindingID);
+					if (BindingReferencesList.Num() == 1)
 					{
-						TArrayView<const TSubclassOf<UMovieSceneCustomBinding>> PrioritySortedCustomBindingTypes = Sequencer->GetSupportedCustomBindingTypes();
-						for (const TSubclassOf<UMovieSceneCustomBinding>& CustomBindingType : PrioritySortedCustomBindingTypes)
+						const FMovieSceneBindingReference& CurrentBindingReference = BindingReferencesList[0];
+						UObject* CurrentBoundObject = Sequencer->FindSpawnedObjectOrTemplate(ObjectBindingID);
+						if (CurrentBoundObject)
 						{
-							if (CustomBindingType
-								&& CurrentBindingReference.CustomBinding
-								&& CurrentBindingReference.CustomBinding->GetClass() != CustomBindingType
-								&& CustomBindingType->GetDefaultObject<UMovieSceneCustomBinding>()->SupportsConversionFromBinding(CurrentBindingReference, CurrentBoundObject))
+							TArrayView<const TSubclassOf<UMovieSceneCustomBinding>> PrioritySortedCustomBindingTypes = Sequencer->GetSupportedCustomBindingTypes();
+							for (const TSubclassOf<UMovieSceneCustomBinding>& CustomBindingType : PrioritySortedCustomBindingTypes)
 							{
-								MenuBuilder.AddMenuEntry(
-									FText::Format(LOCTEXT("ConvertToCustomBinding", "Convert to {0}"), CustomBindingType->GetDefaultObject<UMovieSceneCustomBinding>()->GetBindingTypePrettyName()),
-									FText::Format(LOCTEXT("ConvertToCustomBindingTooltip", "Convert selected binding to {0}"), CustomBindingType->GetDefaultObject<UMovieSceneCustomBinding>()->GetBindingTypePrettyName()),
-									FSlateIcon(),
-									FUIAction(FExecuteAction::CreateRaw(this, &FLevelSequenceCustomization::ConvertToCustomBinding, ObjectBindingModel, CustomBindingType)));
+								if (CustomBindingType
+									&& CurrentBindingReference.CustomBinding
+									&& CurrentBindingReference.CustomBinding->GetClass() != CustomBindingType
+									&& CustomBindingType->GetDefaultObject<UMovieSceneCustomBinding>()->SupportsConversionFromBinding(CurrentBindingReference, CurrentBoundObject))
+								{
+									MenuBuilder.AddMenuEntry(
+										FText::Format(LOCTEXT("ConvertToCustomBinding", "Convert to {0}"), CustomBindingType->GetDefaultObject<UMovieSceneCustomBinding>()->GetBindingTypePrettyName()),
+										FText::Format(LOCTEXT("ConvertToCustomBindingTooltip", "Convert selected binding to {0}"), CustomBindingType->GetDefaultObject<UMovieSceneCustomBinding>()->GetBindingTypePrettyName()),
+										FSlateIcon(),
+										FUIAction(FExecuteAction::CreateRaw(this, &FLevelSequenceCustomization::ConvertToCustomBinding, ObjectBindingModel, CustomBindingType)));
+								}
 							}
 						}
 					}
