Sequencer: Add Force Evaluate BP/Python function to make sure we evaluate immediately, this is needed for example with control rig's since they get setup lazily on evaluation.
#prefglight 6627ec2e0630b3846a063537

diff --git a/Engine/Plugins/MovieScene/LevelSequenceEditor/Source/LevelSequenceEditor/Private/LevelSequenceEditorBlueprintLibrary.cpp b/Engine/Plugins/MovieScene/LevelSequenceEditor/Source/LevelSequenceEditor/Private/LevelSequenceEditorBlueprintLibrary.cpp
--- a/Engine/Plugins/MovieScene/LevelSequenceEditor/Source/LevelSequenceEditor/Private/LevelSequenceEditorBlueprintLibrary.cpp
+++ b/Engine/Plugins/MovieScene/LevelSequenceEditor/Source/LevelSequenceEditor/Private/LevelSequenceEditorBlueprintLibrary.cpp
@@ -634,7 +634,15 @@ void ULevelSequenceEditorBlueprintLibrary::RefreshCurrentLevelSequence()
 		CurrentSequencer.Pin()->NotifyMovieSceneDataChanged(EMovieSceneDataChangeType::Unknown);
 	}
 }
-	
+
+void ULevelSequenceEditorBlueprintLibrary::ForceUpdate()
+{
+	if (CurrentSequencer.IsValid())
+	{
+		CurrentSequencer.Pin()->NotifyMovieSceneDataChanged(EMovieSceneDataChangeType::RefreshAllImmediately);
+	}
+}
+
 TArray<UObject*> ULevelSequenceEditorBlueprintLibrary::GetBoundObjects(FMovieSceneObjectBindingID ObjectBinding)
 {
 	TArray<UObject*> BoundObjects;
diff --git a/Engine/Plugins/MovieScene/LevelSequenceEditor/Source/LevelSequenceEditor/Public/LevelSequenceEditorBlueprintLibrary.h b/Engine/Plugins/MovieScene/LevelSequenceEditor/Source/LevelSequenceEditor/Public/LevelSequenceEditorBlueprintLibrary.h
--- a/Engine/Plugins/MovieScene/LevelSequenceEditor/Source/LevelSequenceEditor/Public/LevelSequenceEditorBlueprintLibrary.h
+++ b/Engine/Plugins/MovieScene/LevelSequenceEditor/Source/LevelSequenceEditor/Public/LevelSequenceEditorBlueprintLibrary.h
@@ -235,10 +235,14 @@ public:
 
 public:
 
-	/** Refresh Sequencer UI. */
+	/** Refresh Sequencer UI on next tick */
 	UFUNCTION(BlueprintCallable, Category = "Level Sequence Editor")
 	static void RefreshCurrentLevelSequence();
 
+	/** Force sequencer evaluation and UI update immediately */
+	UFUNCTION(BlueprintCallable, Category = "Level Sequence Editor")
+	static void ForceUpdate();
+
 	/** Get the object bound to the given binding ID with the current level sequence editor */
 	UFUNCTION(BlueprintPure, Category="Level Sequence Editor")
 	static TArray<UObject*> GetBoundObjects(FMovieSceneObjectBindingID ObjectBinding);
