Sequencer: Fix crash on right click menu

#rnx

diff --git a/Engine/Plugins/MovieScene/LevelSequenceEditor/Source/LevelSequenceEditor/Private/LevelSequenceEditorSubsystem.cpp b/Engine/Plugins/MovieScene/LevelSequenceEditor/Source/LevelSequenceEditor/Private/LevelSequenceEditorSubsystem.cpp
--- a/Engine/Plugins/MovieScene/LevelSequenceEditor/Source/LevelSequenceEditor/Private/LevelSequenceEditorSubsystem.cpp
+++ b/Engine/Plugins/MovieScene/LevelSequenceEditor/Source/LevelSequenceEditor/Private/LevelSequenceEditorSubsystem.cpp
@@ -2501,7 +2501,7 @@ bool ULevelSequenceEditorSubsystem::IsSelectedBindingRootPossessable()
 	{
 		TArray<FGuid> ObjectBindings;
 		Sequencer->GetSelectedObjects(ObjectBindings);
-		if (ObjectBindings.Num() >= 0)
+		if (ObjectBindings.Num() > 0)
 		{
 			UMovieSceneSequence* Sequence = Sequencer->GetFocusedMovieSceneSequence();
 			UMovieScene* MovieScene = Sequence ? Sequence->GetMovieScene() : nullptr;
