Sequencer- Dynamic Binding Refactor
Refactor the 'Dynamic Binding' blueprint system to be Custom Bindings rather than tied directly to FMovieScenePossessable and FMovieSceneSpawnable. As part of this:
? Created UMovieSceneReplaceableDirectorBlueprintBinding and UMovieSceneSpawnableDirectorBlueprintBinding to replace the Possessable and Spawnable versions of DynamicBindings
? Deprecated the DynamicBinding references on possessable and spawnable and created Postload fixup code to automatically create custom bindings instead
? Note that because the previous Dynamic Bindings operated as overrides and the new ones don't, there is a no longer supported path which is using Possessables as a preview in editor for Possessable Dynamic Bindings. Now if you don't mark 'Call in Editor' for a Replaceable Director Blueprint binding, you'll need to define a spawnable preview binding for the binding in order to work with it in Sequencer.
? Since UMG also used Dynamic Bindings on possessables, temporarily add DynamicBinding as a member variable on FWidgetAnimationBinding and create postload logic to move DynamicBindings to live on there for WidgetAnimations.
? USD has custom code to create dynamic bindings automatically for its use case. Convert this also to use the new custom binding system.
? Some cleanup on object binding right-click menus was done as part of this.
? Some further MovieSceneSequence API deprecation and refactoring has been done for the creation of bindings to add SharedPlaybackState where possible which was necessary for dynamic binding setup.

FMovieSceneSpawnable deprecation in Level Sequences
Since I was having to write postload fixup logic for FMovieSceneSpawnables with Dynamic Bindings that converted these to FMovieScenePossessables with custom bindings, it felt like a good chance to deprecate all FMovieSceneSpawnable use in Level Sequences.

I chose at this time not to deprecate FMovieSceneSpawnable use for non-Level-Sequence Movie Scene Sequences. This includes Day Sequences, Template Sequences, Actor Sequences, Widget Animations, Niagara Sequences, and some Meta Human sequences. I made this choice because to do so would require also refactoring those sequence types to use MovieSceneBindingReferences in order to enable custom bindings for those sequence types. This will be a future refactor.

Some notes on this deprecation.
? I kept MovieScene specific code that references or creates FMovieSceneSpawnables due to not refactoring all MovieSceneSequences
? Binding creation code in tools such as Sequencer will check for the presence of Binding References and create custom bindings if present, otherwise will still create FMovieSceneSpawnables
? AvaSequences required some extra work to make this change, including some changes and bugfixes to copy/pasting bindings to move away from custom avalanche copy/paste code, as the copy/paste code for custom bindings is a bit intricate and I didn't want to duplicate it.
? Chaos Cache has been moved away from using an object spawner and instead uses a custom spawnable actor binding inheriting from the regular spawnable actor binding.
? Actor and Take Recording has been updated to support custom spawnable bindings
? Meta Human specific code has been updated to support custom spawnable bindings when using Level Sequences
? FortniteMainBranchObjectVersion has been bumped

Crash fix
As part of the USD work, I found a crash bug where the Binding Properties menu would hold a reference to a level sequence. In the case of USD, the outer of the sequences can be the level, and this caused issues when switching levels. To fix this I ended up needing to add some events when menus are being dismissed so I could clean up these references.

diff --git a/Engine/Plugins/Animation/ControlRig/Source/ControlRigEditor/Private/ControlRigEditorModule.cpp b/Engine/Plugins/Animation/ControlRig/Source/ControlRigEditor/Private/ControlRigEditorModule.cpp
--- a/Engine/Plugins/Animation/ControlRig/Source/ControlRigEditor/Private/ControlRigEditorModule.cpp
+++ b/Engine/Plugins/Animation/ControlRig/Source/ControlRigEditor/Private/ControlRigEditorModule.cpp
@@ -97,6 +97,8 @@
 #include "Constraints/TransformConstraintChannelInterface.h"
 #include "Async/TaskGraphInterfaces.h"
 #include "AssetRegistry/IAssetRegistry.h"
+#include "SequencerUtilities.h"
+#include "Bindings/MovieSceneSpawnableActorBinding.h"
 
 #define LOCTEXT_NAMESPACE "ControlRigEditorModule"
 
@@ -1163,10 +1165,10 @@ void FControlRigEditorModule::BakeToControlRig(UClass* ControlRigClass, UAnimSeq
 			// By default, convert this to a spawnable and delete the existing actor. If for some reason, 
 			// the spawnable couldn't be generated, use the existing actor as a possessable (this could 
 			// eventually be an option)
-			TArray<FGuid> SpawnableGuids = SequencerPtr->ConvertToSpawnable(ActorTrackGuid);
-			if (SpawnableGuids.Num())
-			{	
-				ActorTrackGuid = SpawnableGuids[0];
+
+			if (FMovieScenePossessable* Possessable = FSequencerUtilities::ConvertToCustomBinding(SequencerPtr.ToSharedRef(), ActorTrackGuid, UMovieSceneSpawnableActorBinding::StaticClass(), 0))
+			{ 
+				ActorTrackGuid = Possessable->GetGuid();
 
 				UObject* SpawnedMesh = SequencerPtr->FindSpawnedObjectOrTemplate(ActorTrackGuid);
 
diff --git a/Engine/Plugins/Animation/ControlRig/Source/ControlRigEditor/Private/Sequencer/ControlRigParameterTrackEditor.cpp b/Engine/Plugins/Animation/ControlRig/Source/ControlRigEditor/Private/Sequencer/ControlRigParameterTrackEditor.cpp
--- a/Engine/Plugins/Animation/ControlRig/Source/ControlRigEditor/Private/Sequencer/ControlRigParameterTrackEditor.cpp
+++ b/Engine/Plugins/Animation/ControlRig/Source/ControlRigEditor/Private/Sequencer/ControlRigParameterTrackEditor.cpp
@@ -4184,9 +4184,12 @@ bool FControlRigParameterTrackEditor::HandleAssetAdded(UObject* Asset, const FGu
 
 	const FScopedTransaction Transaction(LOCTEXT("AddControlRigAsset", "Add Control Rig"));
 
-	FGuid NewGuid = GetSequencer()->MakeNewSpawnable(*ASkeletalMeshActor::StaticClass());
+	UE::Sequencer::FCreateBindingParams CreateBindingParams;
+	CreateBindingParams.bSpawnable = true;
+	CreateBindingParams.bAllowCustomBinding = true;
+	FGuid NewGuid = GetSequencer()->CreateBinding(*ASkeletalMeshActor::StaticClass(), CreateBindingParams);
 
-	// MakeNewSpawnable can fail if spawnables are not allowed
+	// CreateBinding can fail if spawnables are not allowed
 	if (!NewGuid.IsValid())
 	{
 		return false;
@@ -4205,8 +4208,7 @@ bool FControlRigParameterTrackEditor::HandleAssetAdded(UObject* Asset, const FGu
 	
 	// Save Spawnable state as the default (with new name and skeletal mesh asset)
 	{
-		FMovieSceneSpawnable* Spawnable = MovieScene->FindSpawnable(NewGuid);
-		GetSequencer()->GetSpawnRegister().SaveDefaultSpawnableState(*Spawnable, GetSequencer()->GetFocusedTemplateID(), *GetSequencer());
+		GetSequencer()->GetSpawnRegister().SaveDefaultSpawnableState(NewGuid, GetSequencer()->GetFocusedTemplateID(), GetSequencer()->GetSharedPlaybackState());
 	}
 
 	UMovieSceneControlRigParameterTrack* Track = Cast<UMovieSceneControlRigParameterTrack>(MovieScene->FindTrack(UMovieSceneControlRigParameterTrack::StaticClass(), NewGuid, NAME_None));
diff --git a/Engine/Plugins/Experimental/Avalanche/Source/AvalancheSequence/Private/AvaSequence.cpp b/Engine/Plugins/Experimental/Avalanche/Source/AvalancheSequence/Private/AvaSequence.cpp
--- a/Engine/Plugins/Experimental/Avalanche/Source/AvalancheSequence/Private/AvaSequence.cpp
+++ b/Engine/Plugins/Experimental/Avalanche/Source/AvalancheSequence/Private/AvaSequence.cpp
@@ -429,7 +429,7 @@ TArray<UObject*> UAvaSequence::GetBoundObjects(UObject* InPlaybackContext) const
 	for (const FMovieSceneBinding& Binding : MovieScene->GetBindings())
 	{
 		TArray<UObject*, TInlineAllocator<1>> BoundObject;
-		LocateBoundObjects(Binding.GetObjectGuid(), UE::UniversalObjectLocator::FResolveParams(InPlaybackContext), BoundObject);
+		LocateBoundObjects(Binding.GetObjectGuid(), UE::UniversalObjectLocator::FResolveParams(InPlaybackContext), MovieSceneHelpers::CreateTransientSharedPlaybackState(InPlaybackContext, MovieScene->GetTypedOuter<UMovieSceneSequence>()), BoundObject);
 		OutObjects.Append(BoundObject);
 	}
 
diff --git a/Engine/Plugins/Experimental/Avalanche/Source/AvalancheSequencer/Private/AvaSequencer.cpp b/Engine/Plugins/Experimental/Avalanche/Source/AvalancheSequencer/Private/AvaSequencer.cpp
--- a/Engine/Plugins/Experimental/Avalanche/Source/AvalancheSequencer/Private/AvaSequencer.cpp
+++ b/Engine/Plugins/Experimental/Avalanche/Source/AvalancheSequencer/Private/AvaSequencer.cpp
@@ -1199,7 +1199,7 @@ UObject* FAvaSequencer::FindResolutionContext(UAvaSequence& InSequence
 		{
 			TArray<UObject*, TInlineAllocator<1>> BoundObjects;
 
-			InSequence.LocateBoundObjects(InGuid, UE::UniversalObjectLocator::FResolveParams(InContextChecked), BoundObjects);
+			InSequence.LocateBoundObjects(InGuid, UE::UniversalObjectLocator::FResolveParams(InContextChecked), MovieSceneHelpers::CreateTransientSharedPlaybackState(InContextChecked, &InSequence), BoundObjects);
 			return BoundObjects;
 		};
 
diff --git a/Engine/Plugins/Experimental/Avalanche/Source/AvalancheSequencer/Private/Clipboard/AvaSequenceExporter.cpp b/Engine/Plugins/Experimental/Avalanche/Source/AvalancheSequencer/Private/Clipboard/AvaSequenceExporter.cpp
--- a/Engine/Plugins/Experimental/Avalanche/Source/AvalancheSequencer/Private/Clipboard/AvaSequenceExporter.cpp
+++ b/Engine/Plugins/Experimental/Avalanche/Source/AvalancheSequencer/Private/Clipboard/AvaSequenceExporter.cpp
@@ -1,7 +1,6 @@
 // Copyright Epic Games, Inc. All Rights Reserved.
 
 #include "AvaSequenceExporter.h"
-#include "AvaSequenceCopyableBinding.h"
 #include "AvaSequencer.h"
 #include "GameFramework/Actor.h"
 #include "MovieScene.h"
@@ -9,14 +8,15 @@
 #include "AvaSequence.h"
 #include "UnrealExporter.h"
 #include "UObject/Package.h"
+#include "SequencerUtilities.h"
 
 namespace UE::AvaSequencer::Private
 {
 	class FAvaSequenceExportObjectInnerContext : public FExportObjectInnerContext
 	{
 	public:
-		explicit FAvaSequenceExportObjectInnerContext(UObject* InPlaybackContext)
-			: PlaybackContext(InPlaybackContext)
+		explicit FAvaSequenceExportObjectInnerContext(const TSharedRef<FAvaSequencer>& InAvaSequencer)
+			: AvaSequencerWeak(InAvaSequencer)
 		{
 		}
 		virtual ~FAvaSequenceExportObjectInnerContext() override = default;
@@ -32,12 +32,14 @@ namespace UE::AvaSequencer::Private
 
 		TConstArrayView<AActor*> GetBoundActors() const { return BoundActors; }
 
-		UObject* GetPlaybackContext() const { return PlaybackContext; }
+		UObject* GetPlaybackContext() const { return AvaSequencerWeak.IsValid() ? AvaSequencerWeak.Pin()->GetPlaybackContext() : nullptr; }
 
-	private:
-		TObjectPtr<UObject> PlaybackContext;
+		TSharedPtr<FAvaSequencer> GetSequencer() const { return AvaSequencerWeak.Pin(); }
 
+	private:
 		TArray<TObjectPtr<AActor>> BoundActors;
+
+		TWeakPtr<FAvaSequencer> AvaSequencerWeak;
 	};
 }
 
@@ -73,7 +75,7 @@ void FAvaSequenceExporter::ExportText(FString& InOutCopiedData, TConstArrayView<
 	}
 
 	UObject* const PlaybackContext = AvaSequencer->GetPlaybackContext();
-	UE::AvaSequencer::Private::FAvaSequenceExportObjectInnerContext ExportContext(PlaybackContext);
+	UE::AvaSequencer::Private::FAvaSequenceExportObjectInnerContext ExportContext(AvaSequencer.ToSharedRef());
 
 	const TCHAR* const Filetype = TEXT("copy");
 	const uint32 PortFlags = PPF_DeepCompareInstances | PPF_ExportsNotFullyQualified;
@@ -118,6 +120,14 @@ bool UAvaSequenceExporter::ExportText(const FExportObjectInnerContext* InContext
 	const UE::AvaSequencer::Private::FAvaSequenceExportObjectInnerContext& Context
 		= static_cast<const UE::AvaSequencer::Private::FAvaSequenceExportObjectInnerContext&>(*InContext);
 
+	TSharedPtr<FAvaSequencer> Sequencer = Context.GetSequencer();
+	 
+	if (!Sequencer)
+	{
+		return false;
+	}
+
+
 	UObject* const PlaybackContext = Context.GetPlaybackContext();
 
 	// Gather guids for the object nodes and any child object nodes
@@ -162,123 +172,10 @@ bool UAvaSequenceExporter::ExportText(const FExportObjectInnerContext* InContext
 		, FCString::Spc(TextIndent)
 		, *Sequence->GetLabel().ToString());
 
-	UObject* const CopyableBindingOuter = GetTransientPackage();
-
-	TArray<UAvaSequenceCopyableBinding*> ObjectsToExport;
-	ObjectsToExport.Reserve(Bindings.Num());
-
-	for (const FMovieSceneBindingProxy& ObjectBinding : Bindings)
-	{
-		UMovieScene* const MovieScene = ObjectBinding.GetMovieScene();
-		if (!IsValid(MovieScene))
-		{
-			continue;
-		}
-
-		UAvaSequenceCopyableBinding* const CopyableBinding = NewObject<UAvaSequenceCopyableBinding>(CopyableBindingOuter, NAME_None, RF_Transient);
-
-		ObjectsToExport.Add(CopyableBinding);
-
-		if (FMovieScenePossessable* const Possessable = MovieScene->FindPossessable(ObjectBinding.BindingID))
-		{
-			UObject* ResolutionContext = FAvaSequencer::FindResolutionContext(*Sequence
-				, *MovieScene
-				, Possessable->GetParent()
-				, PlaybackContext);
-
-			CopyableBinding->Possessable = *Possessable;
-
-			TArray<UObject*, TInlineAllocator<1>> BoundObjects;
-			Sequence->LocateBoundObjects(CopyableBinding->Possessable.GetGuid(), UE::UniversalObjectLocator::FResolveParams(ResolutionContext), BoundObjects);
-
-			// Store the names of the bound objects so that they can be found on paste
-			for (UObject* const BoundObject : BoundObjects)
-			{
-				if (!BoundObject)
-				{
-					continue;
-				}
-
-				if (AActor* const Actor = Cast<AActor>(BoundObject))
-				{
-					CopyableBinding->BoundActorNames.Add(Actor->GetFName());
-				}
-				else if (ResolutionContext != PlaybackContext)
-				{
-					CopyableBinding->BoundObjectPaths.Add(BoundObject->GetPathName(ResolutionContext));
-				}
-			}
-		}
-		else
-		{
-			FMovieSceneSpawnable* const Spawnable = MovieScene->FindSpawnable(ObjectBinding.BindingID);
-			if (Spawnable)
-			{
-				CopyableBinding->Spawnable = *Spawnable;
-
-				// We manually serialize the spawnable object template so that it's not a reference to a privately owned object. Spawnables all have unique copies of their template objects anyways.
-				// Object Templates are re-created on paste (based on these templates) with the correct ownership set up.
-				CopyableBinding->SpawnableObjectTemplate = Spawnable->GetObjectTemplate();
-			}
-		}
-
-		if (const FMovieSceneBinding* const Binding = MovieScene->FindBinding(ObjectBinding.BindingID))
-		{
-			CopyableBinding->Binding = *Binding;
-			for (UMovieSceneTrack* const Track : Binding->GetTracks())
-			{
-				// Tracks suffer from the same issues as Spawnable's Object Templates (reference to a privately owned object). We'll manually serialize the tracks to copy them,
-				// and then restore them on paste.
-				UMovieSceneTrack* const DuplicatedTrack = Cast<UMovieSceneTrack>(StaticDuplicateObject(Track, CopyableBinding));
-				CopyableBinding->Tracks.Add(DuplicatedTrack);
-			}
-		}
-	}
-
-	UAvaSequenceExporter::ExportBindings(ObjectsToExport, TEXT("copy"), TextIndent + 3, Ar, InWarn);
+	TArray<UMovieSceneFolder*> Folders;
+	FSequencerUtilities::CopyBindings(Sequencer->GetSequencer(), Bindings, Folders, Ar);
 
 	Ar.Logf(TEXT("%sEnd Sequence\r\n"), FCString::Spc(TextIndent));
 
 	return true;
 }
-
-void UAvaSequenceExporter::ExportBindings(const TArray<UAvaSequenceCopyableBinding*>& InObjectsToExport, const TCHAR* InType
-	, int32 InTextIndent
-	, FOutputDevice& Ar
-	, FFeedbackContext* InWarn)
-{
-	// Clear the mark state for saving.
-	UnMarkAllObjects(static_cast<EObjectMark>(EObjectMark::OBJECTMARK_TagExp | EObjectMark::OBJECTMARK_TagImp));
-
-	UObject* const Outer = GetTransientPackage();
-
-	const uint32 PortFlags = PPF_ExportsNotFullyQualified | PPF_Copy | PPF_Delimited;
-
-	FExportObjectInnerContext EmptyContext;
-
-	for (UAvaSequenceCopyableBinding* const ObjectToExport : InObjectsToExport)
-	{
-		// We can't use TextExportTransient on USTRUCTS (which our object contains) so we're going to manually null out some references before serializing them. These references are
-		// serialized manually into the archive, as the auto-serialization will only store a reference (to a privately owned object) which creates issues on deserialization. Attempting 
-		// to deserialize these private objects throws a superflous error in the console that makes it look like things went wrong when they're actually OK and expected.
-		TArray<UMovieSceneTrack*> OldTracks = ObjectToExport->Binding.StealTracks(nullptr);
-
-		UObject* const OldSpawnableTemplate = ObjectToExport->Spawnable.GetObjectTemplate();
-
-		ObjectToExport->Spawnable.SetObjectTemplate(nullptr);
-
-		UExporter::ExportToOutputDevice(&EmptyContext, ObjectToExport, nullptr, Ar, InType
-			, InTextIndent, PortFlags, false, Outer);
-
-		// Restore the references (as we don't want to modify the original in the event of a copy operation!)
-		ObjectToExport->Binding.SetTracks(MoveTemp(OldTracks), nullptr);
-		ObjectToExport->Spawnable.SetObjectTemplate(OldSpawnableTemplate);
-
-		// We manually export the object template for the same private-ownership reason as above. Templates need to be re-created anyways as each Spawnable contains its own copy of the template.
-		if (ObjectToExport->SpawnableObjectTemplate)
-		{
-			UExporter::ExportToOutputDevice(&EmptyContext, ObjectToExport->SpawnableObjectTemplate, nullptr, Ar, InType
-				, InTextIndent, PortFlags);
-		}
-	}
-}
diff --git a/Engine/Plugins/Experimental/Avalanche/Source/AvalancheSequencer/Private/Clipboard/AvaSequenceExporter.h b/Engine/Plugins/Experimental/Avalanche/Source/AvalancheSequencer/Private/Clipboard/AvaSequenceExporter.h
--- a/Engine/Plugins/Experimental/Avalanche/Source/AvalancheSequencer/Private/Clipboard/AvaSequenceExporter.h
+++ b/Engine/Plugins/Experimental/Avalanche/Source/AvalancheSequencer/Private/Clipboard/AvaSequenceExporter.h
@@ -35,8 +35,4 @@ public:
 	virtual bool ExportText(const FExportObjectInnerContext* InContext, UObject* InObject, const TCHAR* InType
 		, FOutputDevice& Ar, FFeedbackContext* InWarn, uint32 InPortFlags) override;
 	//~ End UExporter
-
-private:
-	static void ExportBindings(const TArray<UAvaSequenceCopyableBinding*>& InObjectsToExport, const TCHAR* InType
-		, int32 InTextIndent, FOutputDevice& Ar, FFeedbackContext* InWarn);
 };
\ No newline at end of file
diff --git a/Engine/Plugins/Experimental/Avalanche/Source/AvalancheSequencer/Private/Clipboard/AvaSequenceImporter.cpp b/Engine/Plugins/Experimental/Avalanche/Source/AvalancheSequencer/Private/Clipboard/AvaSequenceImporter.cpp
--- a/Engine/Plugins/Experimental/Avalanche/Source/AvalancheSequencer/Private/Clipboard/AvaSequenceImporter.cpp
+++ b/Engine/Plugins/Experimental/Avalanche/Source/AvalancheSequencer/Private/Clipboard/AvaSequenceImporter.cpp
@@ -1,7 +1,6 @@
 // Copyright Epic Games, Inc. All Rights Reserved.
 
 #include "AvaSequenceImporter.h"
-#include "AvaSequenceCopyableBinding.h"
 #include "AvaSequence.h"
 #include "AvaSequencer.h"
 #include "EngineUtils.h"
@@ -11,6 +10,8 @@
 #include "MovieScenePossessable.h"
 #include "SequencerUtilities.h"
 #include "UObject/Package.h"
+#include "Widgets/Notifications/SNotificationList.h"
+#include "Algo/Transform.h"
 
 FAvaSequenceImporter::FAvaSequenceImporter(const TSharedRef<FAvaSequencer>& InAvaSequencer)
 	: AvaSequencerWeak(InAvaSequencer)
@@ -67,8 +68,11 @@ void FAvaSequenceImporter::ImportText(FStringView InPastedData, const TMap<FName
 			BindingsString += TEXT("\r\n");
 		}
 
-		const TArray<UAvaSequenceCopyableBinding*> ImportedBindings = ImportBindings(*Sequencer, BindingsString);
-		ProcessImportedBindings(ImportedBindings, InPastedActors);
+		TArray<FNotificationInfo> PasteErrors; 
+		TArray<FMovieSceneBindingProxy> OutBindings;
+		FMovieScenePasteBindingsParams PasteParams;
+		Algo::Transform(InPastedActors, PasteParams.PastedActors, [](TTuple<FName, AActor*> Pair) { return TTuple<FName, TObjectPtr<AActor>>(Pair.Key, Pair.Value);});
+		FSequencerUtilities::PasteBindings(BindingsString, Sequencer, PasteParams, OutBindings, PasteErrors);
 	}
 
 	AvaSequencer->SetViewedSequence(OriginallyViewedSequence);
@@ -133,427 +137,3 @@ UAvaSequence* FAvaSequenceImporter::GetOrCreateSequence(FAvaSequencer& InAvaSequ
 
 	return NewSequence;
 }
-
-TArray<UAvaSequenceCopyableBinding*> FAvaSequenceImporter::ImportBindings(ISequencer& InSequencerRef, const FString& InBindingsString)
-{
-	class FAvaObjectBindingTextFactory : public FCustomizableTextObjectFactory
-	{
-	public:
-		FAvaObjectBindingTextFactory(ISequencer& InSequencer) : FCustomizableTextObjectFactory(GWarn), Sequencer(InSequencer)
-		{
-		}
-
-		virtual bool CanCreateClass(UClass* InObjectClass, bool& bOmitSubObjs) const override
-		{
-			if (InObjectClass->IsChildOf<UAvaSequenceCopyableBinding>())
-			{
-				return true;
-			}
-			return Sequencer.GetSpawnRegister().CanSpawnObject(InObjectClass);
-		}
-
-		virtual void ProcessConstructedObject(UObject* InObject) override
-		{
-			check(InObject);
-
-			if (InObject->IsA<UAvaSequenceCopyableBinding>())
-			{
-				UAvaSequenceCopyableBinding* CopyableBinding = Cast<UAvaSequenceCopyableBinding>(InObject);
-				ImportedBindings.Add(CopyableBinding);
-			}
-			else
-			{
-				NewSpawnableObjectTemplates.Add(InObject);
-			}
-		}
-
-		TArray<UAvaSequenceCopyableBinding*> ImportedBindings;
-		TArray<UObject*> NewSpawnableObjectTemplates;
-
-	private:
-		ISequencer& Sequencer;
-	};
-
-	UPackage* const TempPackage = NewObject<UPackage>(nullptr, TEXT("/Engine/Editor/AvalancheSequencer/Transient"), RF_Transient);
-	TempPackage->AddToRoot();
-
-	// Turn the text buffer into objects
-	FAvaObjectBindingTextFactory Factory(InSequencerRef);
-	Factory.ProcessBuffer(TempPackage, RF_Transactional, InBindingsString);
-
-	// We had to explicitly serialize object templates due to them being a reference to a privately owned object.
-	// We now deserialize these object template copies
-	// and match them up with their MovieSceneCopyableBinding again.
-	int32 SpawnableObjectTemplateIndex = 0;
-	for (UAvaSequenceCopyableBinding* const ImportedObject : Factory.ImportedBindings)
-	{
-		if (ImportedObject->Spawnable.GetGuid().IsValid() && SpawnableObjectTemplateIndex < Factory.NewSpawnableObjectTemplates.Num())
-		{
-			// This Spawnable Object Template is owned by our transient package, so you'll need to change the owner if you want to keep it later.
-			ImportedObject->SpawnableObjectTemplate = Factory.NewSpawnableObjectTemplates[SpawnableObjectTemplateIndex++];
-		}
-	}
-
-	// Remove the temp package from the root now that it has served its purpose
-	TempPackage->RemoveFromRoot();
-
-	return Factory.ImportedBindings;
-}
-
-void FAvaSequenceImporter::ProcessImportedBindings(const TArray<UAvaSequenceCopyableBinding*>& InImportedBindings
-	, const TMap<FName, AActor*>& InPastedActors)
-{
-	if (InImportedBindings.IsEmpty())
-	{
-		return;
-	}
-
-	TSharedPtr<FAvaSequencer> AvaSequencer = AvaSequencerWeak.Pin();
-	if (!ensureAlways(AvaSequencer.IsValid()))
-	{
-		return;
-	}
-
-	TSharedRef<ISequencer> Sequencer = AvaSequencer->GetSequencer();
-	UAvaSequence* const Sequence     = AvaSequencer->GetViewedSequence();
-	UObject* const PlaybackContext   = Sequencer->GetPlaybackContext();
-
-	if (!ensureAlways(Sequence && PlaybackContext))
-	{
-		return;
-	}
-
-	UMovieScene* const MovieScene = Sequence->GetMovieScene();
-	if (!ensureAlways(MovieScene))
-	{
-		return;
-	}
-
-	struct FImportedPossessable
-	{
-		TArray<FString> BoundObjectPaths;
-	};
-	TMap<FGuid, FImportedPossessable> ImportedPossessableMap;
-	TMap<FGuid, FGuid> OldToNewGuidMap;
-	TArray<FMovieSceneBinding> PastedBindings;
-
-	ImportedPossessableMap.Reserve(InImportedBindings.Num());
-	OldToNewGuidMap.Reserve(InImportedBindings.Num());
-	PastedBindings.Reserve(InImportedBindings.Num());
-
-	FImportContext ImportContext { Sequencer, *Sequence, *MovieScene, Sequencer->GetFocusedTemplateID(), InPastedActors, nullptr };
-
-	for (UAvaSequenceCopyableBinding* const CopyableBinding : InImportedBindings)
-	{
-		// Clear transient flags on the imported tracks
-		for (UMovieSceneTrack* const CopiedTrack : CopyableBinding->Tracks)
-		{
-			FAvaSequenceImporter::ResetCopiedTracksFlags(CopiedTrack);
-		}
-
-		ImportContext.CopyableBinding = CopyableBinding;
-
-		// Possessable Imported Binding
-		if (CopyableBinding->Possessable.GetGuid().IsValid())
-		{
-			FMovieSceneBinding NewBinding = BindPossessable(ImportContext);
-
-			FImportedPossessable ImportedPossessable;
-			ImportedPossessable.BoundObjectPaths = CopyableBinding->BoundObjectPaths;
-			ImportedPossessableMap.Add(NewBinding.GetObjectGuid(), MoveTemp(ImportedPossessable));
-
-			OldToNewGuidMap.Add(CopyableBinding->Possessable.GetGuid(), NewBinding.GetObjectGuid());
-			PastedBindings.Add(MoveTemp(NewBinding));
-		}
-		// Spawnable Imported Binding
-		else if (CopyableBinding->Spawnable.GetGuid().IsValid())
-		{
-			FMovieSceneBinding NewBinding = BindSpawnable(ImportContext);
-			OldToNewGuidMap.Add(CopyableBinding->Spawnable.GetGuid(), NewBinding.GetObjectGuid());
-			PastedBindings.Add(MoveTemp(NewBinding));
-		}
-	}
-
-	// Fix up possessables' parent guids
-	for (const TPair<FGuid, FImportedPossessable>& Pair : ImportedPossessableMap)
-	{
-		FMovieScenePossessable* const Possessable = MovieScene->FindPossessable(Pair.Key);
-		if (!Possessable)
-		{
-			continue;
-		}
-
-		if (FGuid* const NewParentGuid = OldToNewGuidMap.Find(Possessable->GetParent()))
-		{
-			Possessable->SetParent(*NewParentGuid, MovieScene);
-		}
-	}
-
-	// Fix possessables that have bound object paths
-	for (const TPair<FGuid, FImportedPossessable>& Pair : ImportedPossessableMap)
-	{
-		const FGuid& PossessableGuid = Pair.Key;
-		const FImportedPossessable& ImportedPossessable = Pair.Value;
-
-		if (ImportedPossessable.BoundObjectPaths.IsEmpty())
-		{
-			continue;
-		}
-
-		FMovieScenePossessable* const Possessable = MovieScene->FindPossessable(PossessableGuid);
-		if (!Possessable)
-		{
-			continue;
-		}
-
-		auto FindObjectsFunc = [&ImportedPossessableMap, Sequence](const FGuid& InGuid, UObject* InContextChecked)
-			{
-				TArray<UObject*, TInlineAllocator<1>> FoundObjects;
-
-				// Attempt to use the bound object paths that were copied over through the copyable binding
-				// to have an easier time trying to resolve the object
-				if (const FImportedPossessable* const FoundEntry = ImportedPossessableMap.Find(InGuid))
-				{
-					for (const FString& BoundObjectPath : FoundEntry->BoundObjectPaths)
-					{
-						if (UObject* FoundObject = FindObject<UObject>(InContextChecked, *BoundObjectPath))
-						{
-							FoundObjects.Add(FoundObject);
-						}
-					}
-				}
-
-				// If nothing was found, use default method to locate the object
-				if (FoundObjects.IsEmpty() || !FoundObjects[0])
-				{
-					Sequence->LocateBoundObjects(InGuid, UE::UniversalObjectLocator::FResolveParams(InContextChecked), FoundObjects);
-				}
-
-				return FoundObjects;
-			};
-
-		UObject* const ResolutionContext = FAvaSequencer::FindResolutionContext(*Sequence
-			, *MovieScene
-			, Possessable->GetParent()
-			, PlaybackContext
-			, FindObjectsFunc);
-
-		TArray<UObject*> ObjectsToBind;
-		for (const FString& BoundObjectPath : ImportedPossessable.BoundObjectPaths)
-		{
-			if (UObject* FoundObject = FindObject<UObject>(ResolutionContext, *BoundObjectPath))
-			{
-				ObjectsToBind.Add(FoundObject);
-			}
-		}
-
-		AddObjectsToBinding(Sequencer
-			, ObjectsToBind
-			, FMovieSceneBindingProxy(PossessableGuid, Sequence)
-			, ResolutionContext);
-	}
-
-	Sequencer->RestorePreAnimatedState();
-
-	Sequencer->OnMovieSceneBindingsPasted().Broadcast(PastedBindings);
-
-	Sequencer->NotifyMovieSceneDataChanged(EMovieSceneDataChangeType::MovieSceneStructureItemsChanged);
-
-	// Refresh all immediately so that spawned actors will be generated immediately
-	Sequencer->ForceEvaluate();
-
-	// Remap bindings in sections (ie. attach tracks)
-	for (const TPair<FGuid, FGuid>& Pair : OldToNewGuidMap)
-	{
-		FSequencerUtilities::UpdateBindingIDs(Sequencer, Pair.Key, Pair.Value);
-	}
-}
-
-FMovieSceneBinding FAvaSequenceImporter::BindPossessable(const FImportContext& InContext)
-{
-	const FGuid NewGuid = FGuid::NewGuid();
-
-	FMovieSceneBinding NewBinding(NewGuid, InContext.CopyableBinding->Binding.GetName(), InContext.CopyableBinding->Tracks);
-	{
-		FMovieScenePossessable NewPossessable = InContext.CopyableBinding->Possessable;
-		NewPossessable.SetGuid(NewGuid);
-		InContext.MovieScene.AddPossessable(NewPossessable, NewBinding);
-	}
-
-	// Find the actors that this pasted binding should bind to
-	TArray<UObject*> ActorsToBind;
-
-	// Only attempt to bind actor names here, as this is the only thing that can be rebound at this point
-	// Subobjects first need their parents (actors) to resolve correctly, so wait on second pass
-	for (const FName& BoundActorName : InContext.CopyableBinding->BoundActorNames)
-	{
-		if (AActor* const * const FoundActor = InContext.PastedActors.Find(BoundActorName))
-		{
-			AActor* const Actor = *FoundActor;
-			if (!Actor)
-			{
-				continue;
-			}
-
-			// Don't bind if this actor is already bound
-			if (InContext.Sequencer->FindObjectId(*Actor, InContext.SequenceIDRef).IsValid() || ActorsToBind.Contains(Actor))
-			{
-				continue;
-			}
-
-			ActorsToBind.Add(Actor);
-		}
-	}
-
-	// Bind the actors
-	if (!ActorsToBind.IsEmpty())
-	{
-		AddObjectsToBinding(InContext.Sequencer
-			, ActorsToBind
-			, FMovieSceneBindingProxy(NewGuid, &InContext.Sequence)
-			, InContext.Sequencer->GetPlaybackContext());
-	}
-
-	return NewBinding;
-}
-
-FMovieSceneBinding FAvaSequenceImporter::BindSpawnable(const FImportContext& InContext)
-{
-	// We need to let the sequence create the spawnable so that it has everything set up properly internally.
-	// This is required to get spawnables with the correct references to object templates, object templates with
-	// correct owners, etc. However, making a new spawnable also creates the binding for us - this is a problem
-	// because we need to use our binding (which has tracks associated with it). To solve this, we let it create
-	// an object template based off of our (transient package owned) template, then find the newly created binding
-	// and update it.
-
-	FGuid NewGuid;
-	if (InContext.CopyableBinding->SpawnableObjectTemplate)
-	{
-		NewGuid = InContext.Sequencer->MakeNewSpawnable(*InContext.CopyableBinding->SpawnableObjectTemplate
-			, nullptr
-			, false);
-	}
-	else
-	{
-		FMovieSceneSpawnable NewSpawnable{};
-		NewSpawnable.SetGuid(FGuid::NewGuid());
-		NewSpawnable.SetName(InContext.CopyableBinding->Spawnable.GetName());
-
-		InContext.MovieScene.AddSpawnable(NewSpawnable, FMovieSceneBinding(NewSpawnable.GetGuid(), NewSpawnable.GetName()));
-
-		NewGuid = NewSpawnable.GetGuid();
-	}
-
-	FMovieSceneBinding NewBinding(NewGuid, InContext.CopyableBinding->Binding.GetName(), InContext.CopyableBinding->Tracks);
-
-	FMovieSceneSpawnable* const Spawnable = InContext.MovieScene.FindSpawnable(NewGuid);
-
-	// Copy the name of the original spawnable too.
-	Spawnable->SetName(InContext.CopyableBinding->Spawnable.GetName());
-
-	// Clear the transient flags on the copyable binding before assigning to the new spawnable
-	for (UMovieSceneTrack* const Track : NewBinding.GetTracks())
-	{
-		FAvaSequenceImporter::ResetCopiedTracksFlags(Track);
-	}
-
-	// Replace the auto-generated binding with our deserialized bindings (which has our tracks)
-	InContext.MovieScene.ReplaceBinding(NewGuid, NewBinding);
-
-	return NewBinding;
-}
-
-void FAvaSequenceImporter::AddObjectsToBinding(const TSharedRef<ISequencer>& InSequencer
-	, const TArray<UObject*>& InObjectsToAdd
-	, const FMovieSceneBindingProxy& InObjectBinding
-	, UObject* InResolutionContext)
-{
-	UAvaSequence* Sequence = CastChecked<UAvaSequence>(InObjectBinding.Sequence);
-	if (!Sequence)
-	{
-		return;
-	}
-
-	UMovieScene* const MovieScene = Sequence->GetMovieScene();
-	if (!MovieScene || InObjectsToAdd.IsEmpty())
-	{
-		return;
-	}
-
-	UClass* ObjectClass = nullptr;
-	int32 ValidObjectCount = 0;
-
-	FGuid Guid = InObjectBinding.BindingID;
-
-	TArrayView<TWeakObjectPtr<>> WeakObjectsInSequence = Sequence->FindObjectsFromGuid(Guid);
-	for (TWeakObjectPtr<> ObjectWeak : WeakObjectsInSequence)
-	{
-		if (UObject* Object = ObjectWeak.Get())
-		{
-			ObjectClass = Object->GetClass();
-			++ValidObjectCount;
-		}
-	}
-
-	Sequence->Modify();
-	MovieScene->Modify();
-
-	TArray<UObject*> AddedObjects;
-	AddedObjects.Reserve(InObjectsToAdd.Num());
-
-	for (UObject* ObjectToAdd : InObjectsToAdd)
-	{
-		// Skip invalid objects or objects already in the sequence
-		if (!ObjectToAdd || WeakObjectsInSequence.Contains(ObjectToAdd))
-		{
-			continue;
-		}
-
-		// Skip if the object has no common class with the objects already in the binding
-		if (ObjectClass && !UClass::FindCommonBase(ObjectToAdd->GetClass(), ObjectClass))
-		{
-			continue;
-		}
-
-		// if no objects are in the binding, set the class to this object's
-		if (!ObjectClass)
-		{
-			ObjectClass = ObjectToAdd->GetClass();
-		}
-
-		FMovieScenePossessable* Possessable = MovieScene->FindPossessable(Guid);
-		if (!ensureAlways(Possessable))
-		{
-			continue;
-		}
-
-		ObjectToAdd->Modify();
-		if (!Possessable->BindSpawnableObject(InSequencer->GetFocusedTemplateID(), ObjectToAdd, InSequencer->GetSharedPlaybackState()))
-		{
-			Sequence->BindPossessableObject(Guid, *ObjectToAdd, InResolutionContext);
-		}
-		AddedObjects.Add(ObjectToAdd);
-	}
-
-	// Update Labels
-	if (ValidObjectCount + AddedObjects.Num() > 0)
-	{
-		FMovieScenePossessable* const Possessable = MovieScene->FindPossessable(Guid);
-		if (Possessable && ObjectClass)
-		{
-			// If there are multiple objects within the same possessable, name possessable as "ClassName (Count)"
-			if (ValidObjectCount + AddedObjects.Num() > 1)
-			{
-				Possessable->SetName(FString::Printf(TEXT("%s (%d)")
-					, *ObjectClass->GetName()
-					, ValidObjectCount + AddedObjects.Num()));
-			}
-			else if (!AddedObjects.IsEmpty())
-			{
-				Possessable->SetName(FAvaSequencer::GetObjectName(AddedObjects[0]));
-			}
-			Possessable->SetPossessedObjectClass(ObjectClass);
-		}
-	}
-}
-
diff --git a/Engine/Plugins/Experimental/Avalanche/Source/AvalancheSequencer/Private/ObjectBinding/AvaSequencerObjectBinding.cpp b/Engine/Plugins/Experimental/Avalanche/Source/AvalancheSequencer/Private/ObjectBinding/AvaSequencerObjectBinding.cpp
--- a/Engine/Plugins/Experimental/Avalanche/Source/AvalancheSequencer/Private/ObjectBinding/AvaSequencerObjectBinding.cpp
+++ b/Engine/Plugins/Experimental/Avalanche/Source/AvalancheSequencer/Private/ObjectBinding/AvaSequencerObjectBinding.cpp
@@ -75,9 +75,11 @@ void FAvaSequencerObjectBinding::AddPossessActorMenuExtensions(FMenuBuilder& Out
 		{
 			FMovieScenePossessable& Possessable = MovieScene->GetPossessable(Index);
 
+			TArray<UObject*, TInlineAllocator<1>> BoundObjects;
+			static_cast<UMovieSceneSequence*>(Sequence)->LocateBoundObjects(Possessable.GetGuid(), UE::UniversalObjectLocator::FResolveParams(Sequencer->GetPlaybackContext()), Sequencer->GetSharedPlaybackState(), BoundObjects);
+
 			// A possession guid can apply to more than one object, so we get all bound objects for the GUID and add them to our set.
-			ExistingPossessedObjects.Append(static_cast<UMovieSceneSequence*>(Sequence)->LocateBoundObjects(Possessable.GetGuid()
-				, Sequencer->GetPlaybackContext()));
+			ExistingPossessedObjects.Append(BoundObjects);
 		}
 	}
 	
diff --git a/Engine/Plugins/Experimental/ChaosCaching/Source/ChaosCaching/Private/Chaos/ChaosCachingPlugin.cpp b/Engine/Plugins/Experimental/ChaosCaching/Source/ChaosCaching/Private/Chaos/ChaosCachingPlugin.cpp
--- a/Engine/Plugins/Experimental/ChaosCaching/Source/ChaosCaching/Private/Chaos/ChaosCachingPlugin.cpp
+++ b/Engine/Plugins/Experimental/ChaosCaching/Source/ChaosCaching/Private/Chaos/ChaosCachingPlugin.cpp
@@ -5,7 +5,6 @@
 #include "Chaos/Adapters/GeometryCollectionComponentCacheAdapter.h"
 #include "Chaos/Adapters/StaticMeshComponentCacheAdapter.h"
 #include "Chaos/Sequencer/ChaosCacheObjectSpawner.h"
-#include "ILevelSequenceModule.h"
 
 IMPLEMENT_MODULE(IChaosCachingPlugin, ChaosCaching)
 
@@ -16,9 +15,6 @@ void IChaosCachingPlugin::StartupModule()
 	GeometryCollectionAdapter = MakeUnique<Chaos::FGeometryCollectionCacheAdapter>();
 	StaticMeshAdapter = MakeUnique<Chaos::FStaticMeshCacheAdapter>();
 
-	ILevelSequenceModule& LevelSequenceModule = FModuleManager::LoadModuleChecked<ILevelSequenceModule>("LevelSequence");
-	OnCreateMovieSceneObjectSpawnerHandle = LevelSequenceModule.RegisterObjectSpawner(FOnCreateMovieSceneObjectSpawner::CreateStatic(&FChaosCacheObjectSpawner::CreateObjectSpawner));
-	
 	RegisterAdapter(GeometryCollectionAdapter.Get());
 	RegisterAdapter(StaticMeshAdapter.Get());
 }
@@ -30,10 +26,4 @@ void IChaosCachingPlugin::ShutdownModule()
 
 	StaticMeshAdapter = nullptr;
 	GeometryCollectionAdapter = nullptr;
-
-	ILevelSequenceModule* LevelSequenceModule = FModuleManager::GetModulePtr<ILevelSequenceModule>("LevelSequence");
-	if (LevelSequenceModule)
-	{
-		LevelSequenceModule->UnregisterObjectSpawner(OnCreateMovieSceneObjectSpawnerHandle);
-	}
 }
diff --git a/Engine/Plugins/Experimental/ChaosCaching/Source/ChaosCaching/Public/Chaos/Sequencer/ChaosCacheObjectSpawner.h b/Engine/Plugins/Experimental/ChaosCaching/Source/ChaosCaching/Public/Chaos/Sequencer/ChaosCacheObjectSpawner.h
--- a/Engine/Plugins/Experimental/ChaosCaching/Source/ChaosCaching/Public/Chaos/Sequencer/ChaosCacheObjectSpawner.h
+++ b/Engine/Plugins/Experimental/ChaosCaching/Source/ChaosCaching/Public/Chaos/Sequencer/ChaosCacheObjectSpawner.h
@@ -8,14 +8,21 @@
 class FChaosCacheObjectSpawner : public FLevelSequenceActorSpawner
 {
 public:
+	UE_DEPRECATED(5.5, "This has been deprecated as it is no longer necessary. We use UMovieSceneSpawnableChaosCache custom binding instead to handle this logic")
 	CHAOSCACHING_API FChaosCacheObjectSpawner();
 
 	/** Static method to create the object spawner */
+	UE_DEPRECATED(5.5, "This has been deprecated as it is no longer necessary. We use UMovieSceneSpawnableChaosCache custom binding instead to handle this logic")
 	static CHAOSCACHING_API TSharedRef<IMovieSceneObjectSpawner> CreateObjectSpawner();
 
 	// IMovieSceneObjectSpawner interface
+	UE_DEPRECATED(5.5, "This has been deprecated as it is no longer necessary. We use UMovieSceneSpawnableChaosCache custom binding instead to handle this logic")
 	CHAOSCACHING_API virtual UClass* GetSupportedTemplateType() const override;
+
+	UE_DEPRECATED(5.5, "This has been deprecated as it is no longer necessary. We use UMovieSceneSpawnableChaosCache custom binding instead to handle this logic")
 	CHAOSCACHING_API virtual UObject* SpawnObject(FMovieSceneSpawnable& Spawnable, FMovieSceneSequenceIDRef TemplateID, TSharedRef<const UE::MovieScene::FSharedPlaybackState> SharedPlaybackState) override;
+	
+	UE_DEPRECATED(5.5, "This has been deprecated as it is no longer necessary. We use UMovieSceneSpawnableChaosCache custom binding instead to handle this logic")
 	CHAOSCACHING_API virtual void DestroySpawnedObject(UObject& Object) override;
 	virtual bool IsEditor() const override { return true; }
 	virtual int32 GetSpawnerPriority() const override { return 1; }
diff --git a/Engine/Plugins/Experimental/CineCameraRigs/Source/CineCameraRigs/Private/CineCameraRigRail.cpp b/Engine/Plugins/Experimental/CineCameraRigs/Source/CineCameraRigs/Private/CineCameraRigRail.cpp
--- a/Engine/Plugins/Experimental/CineCameraRigs/Source/CineCameraRigs/Private/CineCameraRigRail.cpp
+++ b/Engine/Plugins/Experimental/CineCameraRigs/Source/CineCameraRigs/Private/CineCameraRigRail.cpp
@@ -570,7 +570,7 @@ UMovieSceneFloatTrack* ACineCameraRigRail::FindPositionTrack(const UMovieSceneSe
 	for (const FMovieSceneBinding& Binding : Bindings)
 	{
 		TArray<UObject*, TInlineAllocator<1>> BoundObjects;
-		InSequence->LocateBoundObjects(Binding.GetObjectGuid(), UE::UniversalObjectLocator::FResolveParams(GetWorld()), BoundObjects);
+		InSequence->LocateBoundObjects(Binding.GetObjectGuid(), UE::UniversalObjectLocator::FResolveParams(GetWorld()), MovieSceneHelpers::CreateTransientSharedPlaybackState(GetWorld(), const_cast<UMovieSceneSequence*>(InSequence)), BoundObjects);
 		if (BoundObjects.IsEmpty() || BoundObjects[0] != this)
 		{
 			continue;
diff --git a/Engine/Plugins/Importers/USDImporter/Source/USDExporter/Private/LevelSequenceExporterUSD.cpp b/Engine/Plugins/Importers/USDImporter/Source/USDExporter/Private/LevelSequenceExporterUSD.cpp
--- a/Engine/Plugins/Importers/USDImporter/Source/USDExporter/Private/LevelSequenceExporterUSD.cpp
+++ b/Engine/Plugins/Importers/USDImporter/Source/USDExporter/Private/LevelSequenceExporterUSD.cpp
@@ -50,6 +50,8 @@
 #include "Tracks/MovieSceneSpawnTrack.h"
 #include "Tracks/MovieSceneSubTrack.h"
 #include "UObject/UObjectGlobals.h"
+#include "Bindings/MovieSceneSpawnableDirectorBlueprintBinding.h"
+#include "Bindings/MovieSceneReplaceableDirectorBlueprintBinding.h"
 
 #define LOCTEXT_NAMESPACE "LevelSequenceExporterUSD"
 
@@ -573,16 +575,6 @@ namespace UE::LevelSequenceExporterUSD::Private
 			// Spawn everything for this instance
 			for (FMovieSceneSequenceID SequenceInstance : SequenceInstancePair.Value)
 			{
-				int32 NumSpawnables = MovieScene->GetSpawnableCount();
-				for (int32 Index = 0; Index < NumSpawnables; ++Index)
-				{
-					const FMovieSceneSpawnable& Spawnable = MovieScene->GetSpawnable(Index);
-					const FGuid& Guid = Spawnable.GetGuid();
-
-					StaticCastSharedRef<FMovieSceneSpawnRegister>(Context.SpawnRegister)
-						->SpawnObject(Guid, *MovieScene, SequenceInstance, Context.Sequencer->GetSharedPlaybackState(), 0);
-				}
-
 				if (const FMovieSceneBindingReferences* BindingReferences = Sequence->GetBindingReferences())
 				{
 					int32 BindingIndex = 0;
@@ -857,26 +849,38 @@ namespace UE::LevelSequenceExporterUSD::Private
 							LastGuid = BindingReference.ID;
 							BindingIndex = 0;
 						}
-						if (BindingReference.CustomBinding
-							&& BindingReference.CustomBinding->WillSpawnObject(Context.Sequencer->GetSharedPlaybackState()))
+						if (BindingReference.CustomBinding)
 						{
-							BoundObject = Context.SpawnRegister
-											  ->GetExistingSpawn(*RootSequence, SequenceInstance, FSpawnedInstanceKey(Guid, BindingIndex++));
-							if (!BoundObject)
+							if (BindingReference.CustomBinding->WillSpawnObject(Context.Sequencer->GetSharedPlaybackState()))
 							{
-								// This should never happen as we preemptively spawn everything:
-								// At this point all our spawns should be spawned, but invisible
-								UE_LOG(LogUsd, Warning, TEXT("Failed to find spawned object for spawnable with Guid '%s'"), *Guid.ToString());
-								continue;
+								BoundObject = Context.SpawnRegister->GetExistingSpawn(*RootSequence, SequenceInstance, FSpawnedInstanceKey(Guid, BindingIndex++));
+								if (!BoundObject)
+								{
+									// This should never happen as we preemptively spawn everything:
+									// At this point all our spawns should be spawned, but invisible
+									UE_LOG(LogUsd, Warning, TEXT("Failed to find spawned object for spawnable with Guid '%s'"), *Guid.ToString());
+									continue;
+								}
+
+								if (BoundObject && (BoundObject->IsA<USceneComponent>() || BoundObject->IsA<AActor>()))
+								{
+									BoundObjects.Add(BoundObject, FSpawnedInstanceKey(Guid, BindingIndex));
+								}
 							}
 
-							if (BoundObject && (BoundObject->IsA<USceneComponent>() || BoundObject->IsA<AActor>()))
+							if (UMovieSceneSpawnableDirectorBlueprintBinding* SpawnableDirectorBlueprintBinding = Cast<UMovieSceneSpawnableDirectorBlueprintBinding>(BindingReference.CustomBinding->AsSpawnable(Context.Sequencer->GetSharedPlaybackState())))
 							{
-								BoundObjects.Add(BoundObject, FSpawnedInstanceKey(Guid, BindingIndex));
+								if (SpawnableDirectorBlueprintBinding->DynamicBinding.Function)
+								{
+									DynamicBindings.Add(FSpawnedInstanceKey(Guid, BindingIndex), &SpawnableDirectorBlueprintBinding->DynamicBinding);
+								}
+							}
 
-								if (Possessable.DynamicBinding.Function)
+							if (UMovieSceneReplaceableDirectorBlueprintBinding* ReplaceableDirectorBlueprintBinding = Cast<UMovieSceneReplaceableDirectorBlueprintBinding>(BindingReference.CustomBinding))
+							{
+								if (ReplaceableDirectorBlueprintBinding->DynamicBinding.Function)
 								{
-									DynamicBindings.Add(FSpawnedInstanceKey(Guid, BindingIndex), &Possessable.DynamicBinding);
+									DynamicBindings.Add(FSpawnedInstanceKey(Guid, BindingIndex), &ReplaceableDirectorBlueprintBinding->DynamicBinding);
 								}
 							}
 						}
@@ -904,45 +908,6 @@ namespace UE::LevelSequenceExporterUSD::Private
 				if (BoundObject && (BoundObject->IsA<USceneComponent>() || BoundObject->IsA<AActor>()))
 				{
 					BoundObjects.Add(BoundObject, FSpawnedInstanceKey(Guid, 0));
-
-					if (Possessable.DynamicBinding.Function)
-					{
-						DynamicBindings.Add(FSpawnedInstanceKey(Guid, 0), &Possessable.DynamicBinding);
-					}
-				}
-			}
-
-			// Spawnables
-			int32 NumSpawnables = MovieScene->GetSpawnableCount();
-			for (int32 Index = 0; Index < NumSpawnables; ++Index)
-			{
-				FMovieSceneSpawnable& Spawnable = MovieScene->GetSpawnable(Index);
-				const FGuid& Guid = Spawnable.GetGuid();
-
-				// We won't have spawned ASphereReflectionCapture here.
-				// See the comment inside FLevelSequenceHidingSpawnRegister::SpawnObject and UE-167593 for more info
-				if (ASphereReflectionCapture* ReflectionCapture = Cast<ASphereReflectionCapture>(Spawnable.GetObjectTemplate()))
-				{
-					continue;
-				}
-
-				UObject* BoundObject = Context.SpawnRegister->GetExistingSpawn(*RootSequence, SequenceInstance, FSpawnedInstanceKey(Guid, 0));
-				if (!BoundObject)
-				{
-					// This should never happen as we preemptively spawn everything:
-					// At this point all our spawns should be spawned, but invisible
-					UE_LOG(LogUsd, Warning, TEXT("Failed to find spawned object for spawnable with Guid '%s'"), *Guid.ToString());
-					continue;
-				}
-
-				if (BoundObject && (BoundObject->IsA<USceneComponent>() || BoundObject->IsA<AActor>()))
-				{
-					BoundObjects.Add(BoundObject, FSpawnedInstanceKey(Guid, 0));
-
-					if (Spawnable.DynamicBinding.Function)
-					{
-						DynamicBindings.Add(FSpawnedInstanceKey(Guid, 0), &Spawnable.DynamicBinding);
-					}
 				}
 			}
 		}
@@ -1769,7 +1734,7 @@ bool ULevelSequenceExporterUsd::ExportBinary(
 	// Spawn (but hide) all spawnables so that they will also show up on the level export if we need them to.
 	// We have to traverse the template IDs when spawning spawnables, because we'll want to force each individual spawnable of each
 	// FMovieSceneSequenceID to spawn a separate object, so that they can become separate prims. Without doing this, if we used the same
-	// subsequence with spawnables multiple times within a parent sequence we'd only get one prim out, as the FMovieSceneSpawnable objects
+	// subsequence with spawnables multiple times within a parent sequence we'd only get one prim out, as the spawnable bindings
 	// would be the exact same between all instances of the child sequence (same FGuid)
 	LevelSequenceExporterImpl::PreSpawnSpawnables(Context, *LevelSequence);
 
diff --git a/Engine/Plugins/Importers/USDImporter/Source/USDStage/Private/USDLevelSequenceHelper.cpp b/Engine/Plugins/Importers/USDImporter/Source/USDStage/Private/USDLevelSequenceHelper.cpp
--- a/Engine/Plugins/Importers/USDImporter/Source/USDStage/Private/USDLevelSequenceHelper.cpp
+++ b/Engine/Plugins/Importers/USDImporter/Source/USDStage/Private/USDLevelSequenceHelper.cpp
@@ -218,7 +218,8 @@ namespace UsdLevelSequenceHelperImpl
 			}
 		}
 
-		TArray<UObject*, TInlineAllocator<1>> Objects = MovieSceneSequence.LocateBoundObjects(Guid, ParentContext);
+		TArray<UObject*, TInlineAllocator<1>> Objects;
+		MovieSceneSequence.LocateBoundObjects(Guid, UE::UniversalObjectLocator::FResolveParams(ParentContext), nullptr, Objects);
 		if (Objects.Num() > 0)
 		{
 			return Objects[0];
diff --git a/Engine/Plugins/Importers/USDImporter/Source/USDStage/Private/USDStageActor.cpp b/Engine/Plugins/Importers/USDImporter/Source/USDStage/Private/USDStageActor.cpp
--- a/Engine/Plugins/Importers/USDImporter/Source/USDStage/Private/USDStageActor.cpp
+++ b/Engine/Plugins/Importers/USDImporter/Source/USDStage/Private/USDStageActor.cpp
@@ -91,6 +91,7 @@
 #include "Subsystems/AssetEditorSubsystem.h"
 #include "UnrealEdGlobals.h"
 #include "USDClassesEditorModule.h"
+#include "SequencerUtilities.h"
 #endif	  // WITH_EDITOR
 
 #if USE_USD_SDK
@@ -589,9 +590,10 @@ struct FUsdStageActorImpl
 
 	static void SetupDynamicBinding(
 		const FString& PrimPath,
-		FMovieScenePossessable& Possessable,
+		FGuid BindingID,
 		ULevelSequence* Sequence,
-		const FString& DefaultActorFilter
+		const FString& DefaultActorFilter,
+		TSharedPtr<ISequencer> Sequencer
 	)
 	{
 #if WITH_EDITOR
@@ -602,7 +604,7 @@ struct FUsdStageActorImpl
 		}
 
 		// Make sure we trigger the creation of the DirectorBlueprint before we add our binding, or else
-		// FMovieSceneDynamicBindingInvoker::ResolveDynamicBinding may fail. Note: The regular code path for dynamic
+		// FMovieSceneDynamicBindingUtils::ResolveDynamicBinding may fail. Note: The regular code path for dynamic
 		// bindings via the UI ends up calling these exact same lines from
 		// FMovieSceneDirectorBlueprintEndpointCustomization::PopulateQuickBindSubMenu
 		FMovieSceneSequenceEditor* SequenceEditor = FMovieSceneSequenceEditor::Find(Sequence);
@@ -616,35 +618,54 @@ struct FUsdStageActorImpl
 			return;
 		}
 
+		UMovieSceneReplaceableDirectorBlueprintBinding* NewCustomBinding = Cast<UMovieSceneReplaceableDirectorBlueprintBinding>(UMovieSceneReplaceableDirectorBlueprintBinding::StaticClass()->GetDefaultObject<UMovieSceneCustomBinding>()->CreateNewCustomBinding(nullptr, *MovieScene));
+
+		if (!NewCustomBinding)
+		{
+			return;
+		}
+
+		FMovieSceneDynamicBinding& DynamicBinding = NewCustomBinding->DynamicBinding;
 		// Setup the Dynamic Binding
 		{
-			Possessable.DynamicBinding = FMovieSceneDynamicBinding{};
-			Possessable.DynamicBinding.CompiledFunctionName = NAME_None;
-			Possessable.DynamicBinding.ResolveParamsPinName = TEXT("Params");
-			Possessable.DynamicBinding.Function = UUsdDynamicBindingResolverLibrary::StaticClass()->FindFunctionByName(
+			DynamicBinding = FMovieSceneDynamicBinding{};
+			DynamicBinding.CompiledFunctionName = NAME_None;
+			DynamicBinding.ResolveParamsPinName = TEXT("Params");
+			DynamicBinding.Function = UUsdDynamicBindingResolverLibrary::StaticClass()->FindFunctionByName(
 				GET_FUNCTION_NAME_CHECKED(UUsdDynamicBindingResolverLibrary, ResolveWithStageActor)
 			);
-			if (Possessable.DynamicBinding.Function)
+			if (DynamicBinding.Function)
 			{
-				Possessable.DynamicBinding.ResolveParamsProperty = Possessable.DynamicBinding.Function->FindPropertyByName(
-					Possessable.DynamicBinding.ResolveParamsPinName
+				DynamicBinding.ResolveParamsProperty = DynamicBinding.Function->FindPropertyByName(
+					DynamicBinding.ResolveParamsPinName
 				);
 			}
 
 			// Store a path to this very actor on the binding, so that it can find us later and ask how to resolve a particular
 			// prim path
-			FMovieSceneDynamicBindingPayloadVariable& ActorPathVariable = Possessable.DynamicBinding.PayloadVariables.FindOrAdd(TEXT("StageActorIDNam"
-																																	 "eFilter"));
+			FMovieSceneDynamicBindingPayloadVariable& ActorPathVariable = DynamicBinding.PayloadVariables.FindOrAdd(TEXT("StageActorIDNam"
+				"eFilter"));
 			ActorPathVariable.Value = DefaultActorFilter;
 
-			FMovieSceneDynamicBindingPayloadVariable& RootLayerVariable = Possessable.DynamicBinding.PayloadVariables.FindOrAdd(TEXT("RootLayerFilter"
+			FMovieSceneDynamicBindingPayloadVariable& RootLayerVariable = DynamicBinding.PayloadVariables.FindOrAdd(TEXT("RootLayerFilter"
 			));
 			RootLayerVariable.Value = FString{};	// No root layer filter by default for more flexibility
 
-			FMovieSceneDynamicBindingPayloadVariable& PrimPathVariable = Possessable.DynamicBinding.PayloadVariables.FindOrAdd(TEXT("PrimPath"));
+			FMovieSceneDynamicBindingPayloadVariable& PrimPathVariable = DynamicBinding.PayloadVariables.FindOrAdd(TEXT("PrimPath"));
 			PrimPathVariable.Value = PrimPath;
 		}
 
+		Sequence->Modify();
+		MovieScene->Modify();
+
+		UE::Sequencer::FCreateBindingParams CreateBindingParams;
+		CreateBindingParams.ReplacementGuid = BindingID;
+		CreateBindingParams.BindingIndex = 0;
+		CreateBindingParams.BindingNameOverride = NewCustomBinding->GetDesiredBindingName();
+		CreateBindingParams.CustomBinding = NewCustomBinding;
+
+		FSequencerUtilities::CreateOrReplaceBinding(Sequencer.ToSharedRef(), nullptr, CreateBindingParams);
+
 		// Setup the Blueprint function ("endpoint") on the Sequence's director Blueprint.
 		// In the future we could also just setup a single endpoint for all dynamic bindings, as they're all identical here. However, even then
 		// I think we should allow the option of having this system emit a separatate endpoint for each binding like it does now, as the user
@@ -665,7 +686,12 @@ struct FUsdStageActorImpl
 				return;
 			}
 
-			EndpointDefinition.EndpointName = Possessable.GetName() + TEXT("_DynamicBinding");
+			FMovieScenePossessable* Possessable = Sequence->GetMovieScene()->FindPossessable(BindingID);
+			if (!Possessable)
+			{
+				return;
+			}
+			EndpointDefinition.EndpointName = Possessable->GetName() + TEXT("_DynamicBinding");
 
 			DirectorBlueprint->Modify();
 
@@ -680,7 +706,7 @@ struct FUsdStageActorImpl
 
 		// Create a function call node within that endpoint that calls our UFunction
 		// Reference: FMovieSceneDirectorBlueprintEndpointCustomization::HandleQuickBindActionSelected
-		UBlueprintFunctionNodeSpawner* Spawner = UBlueprintFunctionNodeSpawner::Create(Possessable.DynamicBinding.Function);
+		UBlueprintFunctionNodeSpawner* Spawner = UBlueprintFunctionNodeSpawner::Create(DynamicBinding.Function);
 		FBlueprintActionMenuItem Action{Spawner};
 		UEdGraphPin* ThenPin = NewEndpoint->FindPin(UEdGraphSchema_K2::PN_Then, EGPD_Output);
 		FVector2D NodePosition(NewEndpoint->NodePosX + 400.f, NewEndpoint->NodePosY + 100.0f);
@@ -722,7 +748,7 @@ struct FUsdStageActorImpl
 		// Connect additional parameter pins from the endpoint to the function call node
 		// Reference: FMovieSceneDirectorBlueprintEndpointCustomization::SetEndpoint
 		TSet<FName> NonPayloadPins;
-		const FString* WorldContextParamName = Possessable.DynamicBinding.Function->FindMetaData(FBlueprintMetadata::MD_WorldContext);
+		const FString* WorldContextParamName = DynamicBinding.Function->FindMetaData(FBlueprintMetadata::MD_WorldContext);
 		if (WorldContextParamName)
 		{
 			NonPayloadPins.Add(FName(*WorldContextParamName));
@@ -748,7 +774,7 @@ struct FUsdStageActorImpl
 			}
 		}
 
-		FMovieSceneDynamicBindingUtils::SetEndpoint(MovieScene, &Possessable.DynamicBinding, NewEndpoint);
+		FMovieSceneDynamicBindingUtils::SetEndpoint(MovieScene, &DynamicBinding, NewEndpoint);
 		FMovieSceneDynamicBindingUtils::EnsureBlueprintExtensionCreated(Sequence, DirectorBlueprint);
 		FKismetEditorUtilities::CompileBlueprint(DirectorBlueprint);
 #endif	  // WITH_EDITOR
@@ -4240,7 +4266,7 @@ void AUsdStageActor::OnActorAddedToSequencer(AActor* NewActor, const FGuid Guid,
 	Sequence->Modify();
 	MovieScene->Modify();
 
-	FUsdStageActorImpl::SetupDynamicBinding(PrimPath, *Possessable, Sequence, *GetName());
+	FUsdStageActorImpl::SetupDynamicBinding(PrimPath, Guid, Sequence, *GetName(), PinnedSequencer);
 #endif	  // WITH_EDITOR
 }
 
@@ -4298,7 +4324,7 @@ void AUsdStageActor::OnMovieSceneDataChanged(EMovieSceneDataChangeType ChangeTyp
 	MovieScene->Modify();
 
 	TFunction<UObject*(const FMovieScenePossessable&)> LocateBoundObject = nullptr;
-	LocateBoundObject = [&LocateBoundObject, Sequence, MovieScene](const FMovieScenePossessable& Possessable) -> UObject*
+	LocateBoundObject = [&LocateBoundObject, Sequence, MovieScene, PinnedSequencer](const FMovieScenePossessable& Possessable) -> UObject*
 	{
 		const FGuid& Guid = Possessable.GetGuid();
 		const FGuid& ParentGuid = Possessable.GetParent();
@@ -4312,10 +4338,11 @@ void AUsdStageActor::OnMovieSceneDataChanged(EMovieSceneDataChangeType ChangeTyp
 			}
 		}
 
-		TArray<UObject*, TInlineAllocator<1>> Objects = Cast<UMovieSceneSequence>(Sequence)->LocateBoundObjects(Guid, ParentContext);
+		TArrayView<TWeakObjectPtr<>> Objects = PinnedSequencer->FindBoundObjects(Guid, PinnedSequencer->GetFocusedTemplateID());
+
 		if (Objects.Num() > 0)
 		{
-			return Objects[0];
+			return Objects[0].Get();
 		}
 
 		return nullptr;
@@ -4332,14 +4359,37 @@ void AUsdStageActor::OnMovieSceneDataChanged(EMovieSceneDataChangeType ChangeTyp
 			continue;
 		}
 
+		const FGuid& BindingGuid = Possessable.GetGuid();
+
+		const FMovieSceneBindingReferences* BindingReferences = Sequence->GetBindingReferences();
+		if (!BindingReferences)
+		{
+			continue;
+		}
+
 		// If the binding already has another dynamic binding let's not touch it,
 		// regardless of whether we set that dynamic binding up or the user did
-		if (Possessable.DynamicBinding.Function != nullptr)
+		if (Algo::AnyOf(BindingReferences->GetReferences(Possessable.GetGuid()), [](const FMovieSceneBindingReference& Reference) {
+			if (UMovieSceneSpawnableDirectorBlueprintBinding* SpawnableBinding = Cast<UMovieSceneSpawnableDirectorBlueprintBinding>(Reference.CustomBinding))
+			{
+				if (SpawnableBinding->DynamicBinding.Function != nullptr)
+				{
+					return true;
+				}
+			}
+			if (UMovieSceneReplaceableDirectorBlueprintBinding* ReplaceableBinding = Cast<UMovieSceneReplaceableDirectorBlueprintBinding>(Reference.CustomBinding))
+			{
+				if (ReplaceableBinding->DynamicBinding.Function != nullptr)
+				{
+					return true;
+				}
+			}
+			return false;
+			}))
 		{
 			continue;
 		}
 
-		const FGuid& BindingGuid = Possessable.GetGuid();
 		if (AssetUserData->HandledBindingGuids.Contains(BindingGuid))
 		{
 			continue;
@@ -4361,7 +4411,7 @@ void AUsdStageActor::OnMovieSceneDataChanged(EMovieSceneDataChangeType ChangeTyp
 		// own this component will all return empty strings for SourcePrimPath
 		AssetUserData->HandledBindingGuids.Add(BindingGuid);
 
-		FUsdStageActorImpl::SetupDynamicBinding(SourcePrimPath, Possessable, Sequence, *GetName());
+		FUsdStageActorImpl::SetupDynamicBinding(SourcePrimPath, BindingGuid, Sequence, *GetName(), PinnedSequencer);
 	}
 #endif	  // WITH_EDITOR
 }
diff --git a/Engine/Plugins/MovieScene/LevelSequenceEditor/Source/LevelSequenceEditor/Private/LevelSequenceEditorSubsystem.cpp b/Engine/Plugins/MovieScene/LevelSequenceEditor/Source/LevelSequenceEditor/Private/LevelSequenceEditorSubsystem.cpp
--- a/Engine/Plugins/MovieScene/LevelSequenceEditor/Source/LevelSequenceEditor/Private/LevelSequenceEditorSubsystem.cpp
+++ b/Engine/Plugins/MovieScene/LevelSequenceEditor/Source/LevelSequenceEditor/Private/LevelSequenceEditorSubsystem.cpp
@@ -67,6 +67,7 @@
 #include "DetailLayoutBuilder.h"
 #include "IPropertyTypeCustomization.h"
 #include "IPropertyUtilities.h"
+#include "MovieSceneDynamicBindingCustomization.h"
 
 #include UE_INLINE_GENERATED_CPP_BY_NAME(LevelSequenceEditorSubsystem)
 
@@ -137,6 +138,15 @@ void FMovieSceneBindingPropertyInfoDetailCustomization::CustomizeChildren(TShare
 		FMovieSceneBindingPropertyInfo* BindingPropertyInfo = (StructPtrs.Num() == 1) ? reinterpret_cast<FMovieSceneBindingPropertyInfo*>(StructPtrs[0]) : nullptr;
 		if (BindingPropertyInfo)
 		{
+			bool bShowConvert = true;
+			if (FMovieScenePossessable* Possessable = MovieScene->FindPossessable(BindingGuid))
+			{
+				if (Possessable && Possessable->GetParent().IsValid())
+				{
+					bShowConvert = false;
+				}
+			}
+
 			BindingIndex = InStructPropertyHandle->GetArrayIndex();
 			int32 InitialIndex = 0;
 			BindingTypeNames.Add(MakeShared<FText>(LOCTEXT("BindingType_Possessable", "Possessable")));
@@ -168,6 +178,7 @@ void FMovieSceneBindingPropertyInfoDetailCustomization::CustomizeChildren(TShare
 					SNew(SComboBox<TSharedPtr<FText>>)
 					.OptionsSource(&BindingTypeNames)
 					.InitiallySelectedItem(BindingTypeNames[InitialIndex])
+					.IsEnabled(bShowConvert)
 					.OnSelectionChanged_Lambda([this, InStructPropertyHandle, SupportedBindingTypes, BindingPropertyInfo, &StructBuilder, &CustomizationUtils](TSharedPtr<FText> NewValue, ESelectInfo::Type)
 					{
 						int32 CurrentIndex = BindingTypeNames.IndexOfByPredicate([&](const TSharedPtr<FText>& In) {
@@ -211,8 +222,7 @@ void FMovieSceneBindingPropertyInfoDetailCustomization::CustomizeChildren(TShare
 			];
 
 			// Only show certain menus if we have a currently bound object
-			
-			if (MovieSceneHelpers::GetSingleBoundObject(Sequence, BindingGuid, Sequencer->GetSharedPlaybackState(), BindingIndex) != nullptr)
+			if (bShowConvert && MovieSceneHelpers::GetSingleBoundObject(Sequence, BindingGuid, Sequencer->GetSharedPlaybackState(), BindingIndex) != nullptr)
 			{
 				FDetailWidgetRow& ConvertToRow = StructBuilder.AddCustomRow(LOCTEXT("ConvertBindingTo", "Convert Binding To..."));
 				ConvertToRow.WholeRowContent()
@@ -225,40 +235,40 @@ void FMovieSceneBindingPropertyInfoDetailCustomization::CustomizeChildren(TShare
 						SNew(STextBlock).Text(LOCTEXT("ConvertBindingTo", "Convert Binding To..."))
 					]
 					];
+			}
 
-				// Show Change class and save default state menus for spawnables
-				if (MovieSceneHelpers::IsBoundToSpawnable(Sequence, BindingGuid, Sequencer->GetSharedPlaybackState(), BindingIndex))
-				{
-					FDetailWidgetRow& ChangeClassRow = StructBuilder.AddCustomRow(LOCTEXT("ChangeClass", "Change Class..."));
-					ChangeClassRow.WholeRowContent()
-						[
-							SNew(SComboButton)
-							.OnGetMenuContent(FOnGetContent::CreateLambda([this, &StructBuilder, &CustomizationUtils]() { return OnGetChangeClassMenuContent(StructBuilder, CustomizationUtils); }))
-						.ComboButtonStyle(FAppStyle::Get(), "SimpleComboButton")
-						.ButtonContent()
-						[
-							SNew(STextBlock).Text(LOCTEXT("ChangeClass", "Change Class..."))
-						]
-						];
+			// Show Change class and save default state menus for spawnables
+			if (MovieSceneHelpers::SupportsObjectTemplate(Sequence, BindingGuid, Sequencer->GetSharedPlaybackState(), BindingIndex))
+			{
+				FDetailWidgetRow& ChangeClassRow = StructBuilder.AddCustomRow(LOCTEXT("ChangeClass", "Change Class..."));
+				ChangeClassRow.WholeRowContent()
+					[
+						SNew(SComboButton)
+						.OnGetMenuContent(FOnGetContent::CreateLambda([this, &StructBuilder, &CustomizationUtils]() { return OnGetChangeClassMenuContent(StructBuilder, CustomizationUtils); }))
+					.ComboButtonStyle(FAppStyle::Get(), "SimpleComboButton")
+					.ButtonContent()
+					[
+						SNew(STextBlock).Text(LOCTEXT("ChangeClass", "Change Class..."))
+					]
+					];
 
-					// Save Default State
-					FDetailWidgetRow& SaveDefaultStateRow = StructBuilder.AddCustomRow(LOCTEXT("SaveDefaultState", "Save Default State"));
-					SaveDefaultStateRow.WholeRowContent()
-						[
-							SNew(SButton)
-							.Text(LOCTEXT("SaveDefaultState", "Save Default State"))
-						.ToolTipText(LOCTEXT("SaveDefaultState_Tooltip", "Save the current state of this spawnable as default properties"))
-						.OnClicked_Lambda([this]()
+				// Save Default State
+				FDetailWidgetRow& SaveDefaultStateRow = StructBuilder.AddCustomRow(LOCTEXT("SaveDefaultState", "Save Default State"));
+				SaveDefaultStateRow.WholeRowContent()
+					[
+						SNew(SButton)
+						.Text(LOCTEXT("SaveDefaultState", "Save Default State"))
+					.ToolTipText(LOCTEXT("SaveDefaultState_Tooltip", "Save the current state of this spawnable as default properties"))
+					.OnClicked_Lambda([this]()
+						{
+							TSharedPtr<ISequencer> Sequencer = SequencerPtr.Pin();
+							if (Sequencer.IsValid())
 							{
-								TSharedPtr<ISequencer> Sequencer = SequencerPtr.Pin();
-								if (Sequencer.IsValid())
-								{
-									Sequencer->GetSpawnRegister().SaveDefaultSpawnableState(BindingGuid, BindingIndex, Sequencer->GetFocusedTemplateID(), Sequencer->GetSharedPlaybackState());
-								}
-								return FReply::Handled();
-							})
-						];
-				}
+								Sequencer->GetSpawnRegister().SaveDefaultSpawnableState(BindingGuid, BindingIndex, Sequencer->GetFocusedTemplateID(), Sequencer->GetSharedPlaybackState());
+							}
+							return FReply::Handled();
+						})
+					];
 			}
 
 
@@ -443,11 +453,6 @@ void ULevelSequenceEditorSubsystem::Initialize(FSubsystemCollectionBase& Collect
 	BindingPropertiesMenuExtender = MakeShareable(new FExtender);
 	BindingPropertiesMenuExtender->AddMenuExtension("Possessable", EExtensionHook::First, CommandList, FMenuExtensionDelegate::CreateLambda([this](FMenuBuilder& MenuBuilder) {
 		
-		if (!IsSelectedBindingRootPossessable())
-		{
-			return;
-		}
-		
 		FFormatNamedArguments Args;
 		MenuBuilder.AddSubMenu(
 			FText::Format(LOCTEXT("BindingProperties", "Binding Properties"), Args),
@@ -479,11 +484,7 @@ void ULevelSequenceEditorSubsystem::Initialize(FSubsystemCollectionBase& Collect
 
 	BindingPropertiesMenuExtender->AddMenuExtension("CustomBinding", EExtensionHook::First, CommandList, FMenuExtensionDelegate::CreateLambda([this](FMenuBuilder& MenuBuilder) {
 		// Only add menu entries where the focused sequence is a ULevelSequence
-		if (!IsSelectedBindingRootPossessable())
-		{
-			return;
-		}
-
+		
 		// Add instanced detail customizations
 
 		FFormatNamedArguments Args;
@@ -506,6 +507,8 @@ void ULevelSequenceEditorSubsystem::Deinitialize()
 		SequencerModulePtr->UnregisterOnSequencerCreated(OnSequencerCreatedHandle);
 	}
 
+	FSlateApplication::Get().OnMenuBeingDestroyed().RemoveAll(this);
+
 }
 
 void ULevelSequenceEditorSubsystem::OnSequencerCreated(TSharedRef<ISequencer> InSequencer)
@@ -529,12 +532,31 @@ void ULevelSequenceEditorSubsystem::AddBindingDetailCustomizations(TSharedRef<ID
 				{
 					return MakeShared<FMovieSceneBindingPropertyInfoDetailCustomization>(InSequencer, InMovieScene, InBindingGuid, LevelSequenceEditorSubsystem);
 				}, ActiveSequencer.ToWeakPtr(), MovieScene, BindingGuid, this));
+
+			DetailsView->RegisterInstancedCustomPropertyTypeLayout(FMovieSceneDynamicBinding::StaticStruct()->GetFName(), FOnGetPropertyTypeCustomizationInstance::CreateStatic(&FMovieSceneDynamicBindingCustomization::MakeInstance, MovieScene, BindingGuid));
 			
 			DetailsView->RegisterInstancedCustomPropertyLayout(UMovieSceneSpawnableActorBinding::StaticClass(), FOnGetDetailCustomizationInstance::CreateStatic(&FMovieSceneSpawnableActorBindingBaseCustomization::MakeInstance, ActiveSequencer.ToWeakPtr(), MovieScene, BindingGuid));
 		}
 	}
 }
 
+void ULevelSequenceEditorSubsystem::OnMenuBeingDestroyed(const TSharedRef<IMenu>& Menu, TSharedRef<IDetailsView> DetailsView)
+{
+	TSharedPtr<SWidget> ContentWidget = Menu->GetContent();
+	TSharedPtr<SWidget> ParentWidget = DetailsView;
+	while (ParentWidget)
+	{
+		if (ParentWidget == ContentWidget)
+		{
+			// Binding Properties Menu has closed, clear the binding property list
+			BindingPropertyInfoList = nullptr;
+			FSlateApplication::Get().OnMenuBeingDestroyed().RemoveAll(this);
+			break;
+		}
+		ParentWidget = ParentWidget->GetParentWidget();
+	}
+}
+
 TSharedPtr<ISequencer> ULevelSequenceEditorSubsystem::GetActiveSequencer()
 {
 	for (TWeakPtr<ISequencer> Ptr : Sequencers)
@@ -651,15 +673,22 @@ TArray<FMovieSceneBindingProxy> ULevelSequenceEditorSubsystem::ConvertToSpawnabl
 		return SpawnableProxies;
 	}
 
-	TArray<FMovieSceneSpawnable*> Spawnables = FSequencerUtilities::ConvertToSpawnable(Sequencer.ToSharedRef(), ObjectBinding.BindingID);
-	for (FMovieSceneSpawnable* Spawnable : Spawnables)
+	FMovieScenePossessable* NewPossessable = nullptr;
+
+	if (FMovieSceneBindingReferences* BindingReferences = Sequence->GetBindingReferences())
 	{
-		if (Spawnable)
+		int32 NumBindings = BindingReferences->GetReferences(ObjectBinding.BindingID).Num();
+		for (int32 BindingIndex = 0; BindingIndex < NumBindings; BindingIndex++)
 		{
-			SpawnableProxies.Add(FMovieSceneBindingProxy(Spawnable->GetGuid(), Sequence));
+			NewPossessable = FSequencerUtilities::ConvertToCustomBinding(Sequencer.ToSharedRef(), ObjectBinding.BindingID, UMovieSceneSpawnableActorBinding::StaticClass(), BindingIndex);
 		}
 	}
 
+	if (NewPossessable)
+	{
+		SpawnableProxies.Add(FMovieSceneBindingProxy(NewPossessable->GetGuid(), Sequence));
+	}
+
 	return SpawnableProxies;
 }
 
@@ -2168,6 +2197,8 @@ void ULevelSequenceEditorSubsystem::AddBindingPropertiesMenu(FMenuBuilder& MenuB
 		RefreshBindingDetails(&DetailsView.Get(), ObjectBindings[0]);
 		DetailsView->OnFinishedChangingProperties().AddUObject(this, &ULevelSequenceEditorSubsystem::OnFinishedChangingLocators, DetailsView, ObjectBindings[0]);
 
+		FSlateApplication::Get().OnMenuBeingDestroyed().AddUObject(this, &ULevelSequenceEditorSubsystem::OnMenuBeingDestroyed, DetailsView);
+
 		MenuBuilder.AddWidget(DetailsView, FText::GetEmpty(), true);
 	}
 }
@@ -2231,6 +2262,10 @@ void ULevelSequenceEditorSubsystem::OnFinishedChangingLocators(const FPropertyCh
 				CopiedBinding = Cast<UMovieSceneCustomBinding>(StaticDuplicateObject(LocatorInfo.CustomBinding, MovieScene));
 			}
 			BindingReferences->AddBinding(ObjectBindingID, MoveTemp(LocatorInfo.Locator), LocatorInfo.ResolveFlags, CopiedBinding);
+			if (CopiedBinding)
+			{
+				CopiedBinding->OnBindingAddedOrChanged(*MovieScene);
+			}
 		}
 
 		Sequencer->State.Invalidate(ObjectBindingID, Sequencer->GetFocusedTemplateID());
@@ -2280,6 +2315,12 @@ void ULevelSequenceEditorSubsystem::OnFinishedChangingLocators(const FPropertyCh
 
 		Sequencer->NotifyMovieSceneDataChanged(EMovieSceneDataChangeType::MovieSceneStructureItemsChanged);
 
+		// Destroy any previous spawnables- they'll get recreated on the force evaluate below
+		for (int32 BindingIndex = 0; BindingIndex < BindingPropertyInfoList->Bindings.Num(); ++BindingIndex)
+		{
+			Sequencer->GetSpawnRegister().DestroySpawnedObject(ObjectBindingID, Sequencer->GetFocusedTemplateID(), Sequencer->GetSharedPlaybackState(), BindingIndex);
+		}
+
 		// Force evaluate the Sequencer after clearing the cache (which the above will do) so that any newly loaded actors will be loaded as part of the transaction
 		Sequencer->ForceEvaluate();
 
@@ -2293,10 +2334,10 @@ void ULevelSequenceEditorSubsystem::OnFinishedChangingLocators(const FPropertyCh
 				UMovieSceneCustomBinding* CopiedBinding = nullptr;
 				if (Reference.CustomBinding)
 				{
-					CopiedBinding = Cast<UMovieSceneCustomBinding>(StaticDuplicateObject(Reference.CustomBinding, BindingPropertyInfoList.Get()));
+					CopiedBinding = Cast<UMovieSceneCustomBinding>(StaticDuplicateObject(Reference.CustomBinding, this));
 				}
 
-				return FMovieSceneBindingPropertyInfo{ Reference.Locator, Reference.ResolveFlags, Reference.CustomBinding };
+				return FMovieSceneBindingPropertyInfo{ Reference.Locator, Reference.ResolveFlags, CopiedBinding };
 			});
 
 
@@ -2509,7 +2550,7 @@ bool ULevelSequenceEditorSubsystem::IsSelectedBindingRootPossessable()
 			{
 				if (FMovieScenePossessable* Possessable = MovieScene->FindPossessable(ObjectBindings[0]))
 				{
-					if (!Possessable->GetParent().IsValid())
+					if (!Possessable->GetParent().IsValid() && !Possessable->GetSpawnableObjectBindingID().IsValid())
 					{
 						return true;
 					}
@@ -2539,15 +2580,15 @@ void ULevelSequenceEditorSubsystem::RefreshBindingDetails(IDetailsView* DetailsV
 	}
 	if (FMovieSceneBindingReferences* BindingReferences = Sequence->GetBindingReferences())
 	{
-		BindingPropertyInfoList = NewObject<UMovieSceneBindingPropertyInfoList>();
+		BindingPropertyInfoList = NewObject<UMovieSceneBindingPropertyInfoList>(this);
 		Algo::Transform(BindingReferences->GetReferences(ObjectBindingID), BindingPropertyInfoList->Bindings, [this](const FMovieSceneBindingReference& Reference)
 			{
 				UMovieSceneCustomBinding* CopiedBinding = nullptr;
 				if (Reference.CustomBinding)
 				{
-					CopiedBinding = Cast<UMovieSceneCustomBinding>(StaticDuplicateObject(Reference.CustomBinding, BindingPropertyInfoList.Get()));
+					CopiedBinding = Cast<UMovieSceneCustomBinding>(StaticDuplicateObject(Reference.CustomBinding, this));
 				}
-				return FMovieSceneBindingPropertyInfo{ Reference.Locator, Reference.ResolveFlags, Reference.CustomBinding };
+				return FMovieSceneBindingPropertyInfo{ Reference.Locator, Reference.ResolveFlags, CopiedBinding };
 			});
 
 		DetailsView->SetObject(BindingPropertyInfoList.Get(), true);
diff --git a/Engine/Plugins/MovieScene/LevelSequenceEditor/Source/LevelSequenceEditor/Private/Misc/LevelSequenceCustomization.cpp b/Engine/Plugins/MovieScene/LevelSequenceEditor/Source/LevelSequenceEditor/Private/Misc/LevelSequenceCustomization.cpp
--- a/Engine/Plugins/MovieScene/LevelSequenceEditor/Source/LevelSequenceEditor/Private/Misc/LevelSequenceCustomization.cpp
+++ b/Engine/Plugins/MovieScene/LevelSequenceEditor/Source/LevelSequenceEditor/Private/Misc/LevelSequenceCustomization.cpp
@@ -21,7 +21,6 @@
 #include "ScopedTransaction.h"
 #include "SequencerCommands.h"
 #include "SequencerUtilities.h"
-#include "Tracks/MovieSceneSpawnTrack.h"
 #include "MovieSceneBindingReferences.h"
 #include "Bindings/MovieSceneSpawnableBinding.h"
 #include "ClassViewerFilter.h"
@@ -31,56 +30,6 @@
 namespace UE::Sequencer
 {
 
-struct FMovieSceneSpawnableFlagCheckState
-{
-	TSharedPtr<FSequencerEditorViewModel> EditorViewModel;
-	UMovieScene* MovieScene;
-	bool FMovieSceneSpawnable::* PtrToFlag;
-
-	ECheckBoxState operator()() const
-	{
-		ECheckBoxState CheckState = ECheckBoxState::Undetermined;
-		for (TViewModelPtr<IObjectBindingExtension> ObjectBinding : EditorViewModel->GetSelection()->Outliner.Filter<IObjectBindingExtension>())
-		{
-			FMovieSceneSpawnable* SelectedSpawnable = MovieScene->FindSpawnable(ObjectBinding->GetObjectGuid());
-			if (SelectedSpawnable)
-			{
-				if (CheckState != ECheckBoxState::Undetermined && SelectedSpawnable->*PtrToFlag != (CheckState == ECheckBoxState::Checked))
-				{
-					return ECheckBoxState::Undetermined;
-				}
-				CheckState = SelectedSpawnable->*PtrToFlag ? ECheckBoxState::Checked : ECheckBoxState::Unchecked;
-			}
-		}
-		return CheckState;
-	}
-};
-
-struct FMovieSceneSpawnableFlagToggler
-{
-	TSharedPtr<FSequencerEditorViewModel> EditorViewModel;
-	UMovieScene* MovieScene;
-	bool FMovieSceneSpawnable::* PtrToFlag;
-	FText TransactionText;
-
-	void operator()() const
-	{
-		FScopedTransaction Transaction(TransactionText);
-
-		const ECheckBoxState CheckState = FMovieSceneSpawnableFlagCheckState{ EditorViewModel, MovieScene, PtrToFlag }();
-
-		MovieScene->Modify();
-		for (TViewModelPtr<IObjectBindingExtension> ObjectBinding : EditorViewModel->GetSelection()->Outliner.Filter<IObjectBindingExtension>())
-		{
-			FMovieSceneSpawnable* SelectedSpawnable = MovieScene->FindSpawnable(ObjectBinding->GetObjectGuid());
-			if (SelectedSpawnable)
-			{
-				SelectedSpawnable->*PtrToFlag = (CheckState == ECheckBoxState::Unchecked);
-			}
-		}
-	}
-};
-
 void FLevelSequenceCustomization::RegisterSequencerCustomization(FSequencerCustomizationBuilder& Builder)
 {
 	WeakSequencer = Builder.GetSequencer().AsShared();
@@ -181,96 +130,71 @@ void FLevelSequenceCustomization::ExtendObjectBindingContextMenu(FMenuBuilder& M
 	}
 	bool bShowConvert = true;
 
-	FMovieSceneSpawnable* Spawnable = MovieScene->FindSpawnable(ObjectBindingID);
-
-	if (Spawnable)
-	{
-		MenuBuilder.BeginSection("Spawnable", LOCTEXT("SpawnableMenuSectionName", "Spawnable"));
-
-		MenuBuilder.AddSubMenu(
-			LOCTEXT("OwnerLabel", "Spawned Object Owner"),
-			LOCTEXT("OwnerTooltip", "Specifies how the spawned object is to be owned"),
-			FNewMenuDelegate::CreateRaw(this, &FLevelSequenceCustomization::AddSpawnOwnershipMenu, ObjectBindingModel)
-		);
-
-		MenuBuilder.AddSubMenu(
-			LOCTEXT("SubLevelLabel", "Spawnable Level"),
-			LOCTEXT("SubLevelTooltip", "Specifies which level the spawnable should be spawned into"),
-			FNewMenuDelegate::CreateRaw(this, &FLevelSequenceCustomization::AddSpawnLevelMenu, ObjectBindingModel)
-		);
-
-		MenuBuilder.AddSubMenu(
-			LOCTEXT("ChangeClassLabel", "Change Class"),
-			LOCTEXT("ChangeClassTooltip", "Change the class (object template) that this spawns from"),
-			FNewMenuDelegate::CreateLambda([=](FMenuBuilder& MenuBuilder) { FSequencerUtilities::AddChangeClassMenu(MenuBuilder, Sequencer.ToSharedRef(), ObjectBindingID, 0, TFunction<void()>()); }));
-
-		MenuBuilder.AddMenuEntry(
-			LOCTEXT("ContinuouslyRespawn", "Continuously Respawn"),
-			LOCTEXT("ContinuouslyRespawnTooltip", "When enabled, this spawnable will always be respawned if it gets destroyed externally. When disabled, this object will only ever be spawned once for each spawn key even if it is destroyed externally"),
-			FSlateIcon(),
-			FUIAction(
-				FExecuteAction::CreateLambda(FMovieSceneSpawnableFlagToggler{ EditorViewModel, MovieScene, &FMovieSceneSpawnable::bContinuouslyRespawn, LOCTEXT("ContinuouslyRespawnTransaction", "Set Continuously Respawn") }),
-				FCanExecuteAction(),
-				FGetActionCheckState::CreateLambda(FMovieSceneSpawnableFlagCheckState{ EditorViewModel, MovieScene, &FMovieSceneSpawnable::bContinuouslyRespawn })
-			),
-			NAME_None,
-			EUserInterfaceActionType::ToggleButton
-		);
-
-		MenuBuilder.AddMenuEntry(
-			LOCTEXT("NetAddressable", "Net Addressable"),
-			LOCTEXT("NetAddressableTooltip", "When enabled, this spawnable will be spawned using a unique name that allows it to be addressed by the server and client (useful for relative movement calculations on spawned props)"),
-			FSlateIcon(),
-			FUIAction(
-				FExecuteAction::CreateLambda(FMovieSceneSpawnableFlagToggler{ EditorViewModel, MovieScene, &FMovieSceneSpawnable::bNetAddressableName, LOCTEXT("NetAddressableTransaction", "Set Net Addressable") }),
-				FCanExecuteAction(),
-				FGetActionCheckState::CreateLambda(FMovieSceneSpawnableFlagCheckState{ EditorViewModel, MovieScene, &FMovieSceneSpawnable::bNetAddressableName })
-			),
-			NAME_None,
-			EUserInterfaceActionType::ToggleButton
-		);
-
-		MenuBuilder.AddMenuEntry(FSequencerCommands::Get().SaveCurrentSpawnableState);
-
-		MenuBuilder.AddSubMenu(
-				LOCTEXT("DynamicSpawn", "Dynamic Spawn"),
-				LOCTEXT("DynamicSpawnTooltip", "Specify a Blueprint method that will spawn or otherwise acquire a compatible actor for this binding"),
-				FNewMenuDelegate::CreateRaw(this, &FLevelSequenceCustomization::AddDynamicSpawnMenu, ObjectBindingModel));
-
-		MenuBuilder.EndSection();
-	}
-	else if (FMovieScenePossessable* Possessable = MovieScene->FindPossessable(ObjectBindingID))
+	if (FMovieScenePossessable* Possessable = MovieScene->FindPossessable(ObjectBindingID))
 	{
-		bShowConvert = !Possessable->GetParent().IsValid();
+		// We can't convert sub-objects to different binding types for now.
+		if (Possessable->GetParent().IsValid())
+		{
+			bShowConvert = false;
+		}
 		bool bCustomBinding = false;
 		bool bMultipleBindings = false;
+		UObject* ResolutionContext = MovieSceneHelpers::GetResolutionContext(Sequence, ObjectBindingID, Sequencer->GetFocusedTemplateID(), Sequencer->GetSharedPlaybackState());
+
 		if (const FMovieSceneBindingReferences* BindingReferences = Sequencer->GetFocusedMovieSceneSequence()->GetBindingReferences())
 		{
 			bCustomBinding = Algo::AnyOf(BindingReferences->GetReferences(ObjectBindingID), [](const FMovieSceneBindingReference& Reference) { return Reference.CustomBinding; });
 			bMultipleBindings = BindingReferences->GetReferences(ObjectBindingID).Num() > 1;
+			UE::UniversalObjectLocator::FResolveParams LocatorResolveParams(ResolutionContext);
+			FMovieSceneBindingResolveParams BindingResolveParams{ Sequencer->GetFocusedMovieSceneSequence(), ObjectBindingID, Sequencer->GetFocusedTemplateID(), ResolutionContext };
+			if (!BindingReferences->ResolveSingleBinding(BindingResolveParams, 0, LocatorResolveParams, Sequencer->GetSharedPlaybackState()))
+			{
+				bShowConvert = false;
+			}
+
+			// Can convert to possessable
+			int32 BindingIndex = 0;
+			bool bAnyValidConversions = false;
+			if (Algo::AnyOf(BindingReferences->GetReferences(ObjectBindingID), [&BindingIndex, Sequencer, &Sequence, &MovieScene](const FMovieSceneBindingReference& BindingReference) {
+				return FSequencerUtilities::CanConvertToPossessable(Sequencer.ToSharedRef(), BindingReference.ID, BindingIndex++);
+				}))
+			{
+				bAnyValidConversions = true;
+			}
+			else
+			{
+				TArrayView<const TSubclassOf<UMovieSceneCustomBinding>> PrioritySortedCustomBindingTypes = Sequencer->GetSupportedCustomBindingTypes();
+				for (const TSubclassOf<UMovieSceneCustomBinding>& CustomBindingType : PrioritySortedCustomBindingTypes)
+				{
+					BindingIndex = 0;
+					if (Algo::AllOf(BindingReferences->GetReferences(ObjectBindingID), [&BindingIndex, &Sequence, &MovieScene, &CustomBindingType, Sequencer](const FMovieSceneBindingReference& BindingReference)
+						{
+							return FSequencerUtilities::CanConvertToCustomBinding(Sequencer.ToSharedRef(), BindingReference.ID, CustomBindingType, BindingIndex++);
+						}))
+					{
+						bAnyValidConversions = true;
+						break;
+					}
+				}
+			}
+			if (!bAnyValidConversions)
+			{
+				bShowConvert = false;
+			}
 		}
 
 		// Regular possessable
 		if (!bCustomBinding)
 		{
+			// Regular possessable
+			// We don't add anything here, but the extension will
 			MenuBuilder.BeginSection("Possessable");
-
-			if (!bMultipleBindings)
-			{
-				if (UMovieScene::IsTrackClassAllowed(ULevelSequenceDirector::StaticClass()))
-				{
-					MenuBuilder.AddSubMenu(
-						LOCTEXT("DynamicPossession", "Dynamic Possession"),
-						LOCTEXT("DynamicPossessionTooltip", "Specify a Blueprint method that will find a compatible actor for this binding"),
-						FNewMenuDelegate::CreateRaw(this, &FLevelSequenceCustomization::AddDynamicPossessionMenu, ObjectBindingModel));
-				}
-			}
 			MenuBuilder.EndSection();
 		}
 		else
 		{
 			MenuBuilder.BeginSection("CustomBinding");
-			bool bCustomSpawnable = MovieSceneHelpers::IsBoundToAnySpawnable(Sequencer->GetFocusedMovieSceneSequence(), ObjectBindingID, Sequencer->GetSharedPlaybackState());
+			bool bCustomSpawnable = MovieSceneHelpers::SupportsObjectTemplate(Sequencer->GetFocusedMovieSceneSequence(), ObjectBindingID, Sequencer->GetSharedPlaybackState());
 			// Check for custom binding types
 
 			if (bCustomSpawnable)
@@ -283,21 +207,6 @@ void FLevelSequenceCustomization::ExtendObjectBindingContextMenu(FMenuBuilder& M
 						LOCTEXT("ChangeClassLabel", "Change Class"),
 						LOCTEXT("ChangeClassTooltip", "Change the class (object template) that this spawns from"),
 						FNewMenuDelegate::CreateLambda([=](FMenuBuilder& MenuBuilder) { FSequencerUtilities::AddChangeClassMenu(MenuBuilder, Sequencer.ToSharedRef(), ObjectBindingID, 0, TFunction<void()>()); }));
-				}		
-				
-				MenuBuilder.AddSubMenu(
-					LOCTEXT("DynamicSpawn", "Dynamic Spawn"),
-					LOCTEXT("DynamicSpawnTooltip", "Specify a Blueprint method that will spawn or otherwise acquire a compatible actor for this binding"),
-					FNewMenuDelegate::CreateRaw(this, &FLevelSequenceCustomization::AddDynamicSpawnMenu, ObjectBindingModel));
-			}
-			else
-			{
-				if (UMovieScene::IsTrackClassAllowed(ULevelSequenceDirector::StaticClass()))
-				{
-					MenuBuilder.AddSubMenu(
-						LOCTEXT("DynamicPossession", "Dynamic Possession"),
-						LOCTEXT("DynamicPossessionTooltip", "Specify a Blueprint method that will find a compatible actor for this binding"),
-						FNewMenuDelegate::CreateRaw(this, &FLevelSequenceCustomization::AddDynamicPossessionMenu, ObjectBindingModel));
 				}
 			}
 
@@ -355,201 +264,6 @@ void FLevelSequenceCustomization::ExtendObjectBindingContextMenu(FMenuBuilder& M
 	MenuBuilder.EndSection();
 }
 
-void FLevelSequenceCustomization::AddSpawnOwnershipMenu(FMenuBuilder& MenuBuilder, TSharedPtr<FObjectBindingModel> ObjectBindingModel)
-{
-	TSharedPtr<ISequencer> Sequencer = WeakSequencer.Pin();
-	TSharedPtr<FSequencerEditorViewModel> EditorViewModel = Sequencer->GetViewModel();
-
-	FGuid ObjectBindingID = ObjectBindingModel->GetObjectGuid();
-	UMovieScene* MovieScene = Sequencer->GetFocusedMovieSceneSequence()->GetMovieScene();
-
-	FMovieSceneSpawnable* Spawnable = MovieScene->FindSpawnable(ObjectBindingID);
-	if (!Spawnable)
-	{
-		return;
-	}
-
-	auto Callback = [=](ESpawnOwnership NewOwnership)
-	{
-		FScopedTransaction Transaction(LOCTEXT("SetSpawnOwnership", "Set Spawnable Ownership"));
-
-		MovieScene->Modify();
-
-		for (TViewModelPtr<IObjectBindingExtension> ObjectBinding : Sequencer->GetViewModel()->GetSelection()->Outliner.Filter<IObjectBindingExtension>())
-		{
-			FMovieSceneSpawnable* Spawnable = MovieScene->FindSpawnable(ObjectBinding->GetObjectGuid());
-
-			Spawnable->SetSpawnOwnership(NewOwnership);
-
-			// Overwrite the completion state for all spawn sections to ensure the expected behaviour.
-			EMovieSceneCompletionMode NewCompletionMode = NewOwnership == ESpawnOwnership::InnerSequence ? EMovieSceneCompletionMode::RestoreState : EMovieSceneCompletionMode::KeepState;
-
-			// Make all spawn sections retain state
-			UMovieSceneSpawnTrack* SpawnTrack = MovieScene->FindTrack<UMovieSceneSpawnTrack>(ObjectBinding->GetObjectGuid());
-			if (SpawnTrack)
-			{
-				for (UMovieSceneSection* Section : SpawnTrack->GetAllSections())
-				{
-					Section->Modify();
-					Section->EvalOptions.CompletionMode = NewCompletionMode;
-				}
-			}
-		}
-	};
-
-	MenuBuilder.AddMenuEntry(
-		LOCTEXT("ThisSequence_Label", "This Sequence"),
-		LOCTEXT("ThisSequence_Tooltip", "Indicates that this sequence will own the spawned object. The object will be destroyed at the end of the sequence."),
-		FSlateIcon(),
-		FUIAction(
-			FExecuteAction::CreateLambda(Callback, ESpawnOwnership::InnerSequence),
-			FCanExecuteAction(),
-			FIsActionChecked::CreateLambda([=] { return Spawnable->GetSpawnOwnership() == ESpawnOwnership::InnerSequence; })
-		),
-		NAME_None,
-		EUserInterfaceActionType::ToggleButton
-	);
-
-	MenuBuilder.AddMenuEntry(
-		LOCTEXT("RootSequence_Label", "Root Sequence"),
-		LOCTEXT("RootSequence_Tooltip", "Indicates that the outermost sequence will own the spawned object. The object will be destroyed when the outermost sequence stops playing."),
-		FSlateIcon(),
-		FUIAction(
-			FExecuteAction::CreateLambda(Callback, ESpawnOwnership::RootSequence),
-			FCanExecuteAction(),
-			FIsActionChecked::CreateLambda([=] { return Spawnable->GetSpawnOwnership() == ESpawnOwnership::RootSequence; })
-		),
-		NAME_None,
-		EUserInterfaceActionType::ToggleButton
-	);
-
-	MenuBuilder.AddMenuEntry(
-		LOCTEXT("External_Label", "External"),
-		LOCTEXT("External_Tooltip", "Indicates this object's lifetime is managed externally once spawned. It will not be destroyed by sequencer."),
-		FSlateIcon(),
-		FUIAction(
-			FExecuteAction::CreateLambda(Callback, ESpawnOwnership::External),
-			FCanExecuteAction(),
-			FIsActionChecked::CreateLambda([=] { return Spawnable->GetSpawnOwnership() == ESpawnOwnership::External; })
-		),
-		NAME_None,
-		EUserInterfaceActionType::ToggleButton
-	);
-}
-
-void FLevelSequenceCustomization::AddSpawnLevelMenu(FMenuBuilder& MenuBuilder, TSharedPtr<FObjectBindingModel> ObjectBindingModel)
-{
-	TSharedPtr<ISequencer> Sequencer = WeakSequencer.Pin();
-	TSharedPtr<FSequencerEditorViewModel> EditorViewModel = Sequencer->GetViewModel();
-
-	FGuid ObjectBindingID = ObjectBindingModel->GetObjectGuid();
-	UMovieScene* MovieScene = Sequencer->GetFocusedMovieSceneSequence()->GetMovieScene();
-
-	FMovieSceneSpawnable* Spawnable = MovieScene->FindSpawnable(ObjectBindingID);
-	if (!Spawnable)
-	{
-		return;
-	}
-
-	MenuBuilder.AddMenuEntry(
-		NSLOCTEXT("UnrealEd", "PersistentLevel", "Persistent Level"),
-		NSLOCTEXT("UnrealEd", "PersistentLevel", "Persistent Level"),
-		FSlateIcon(),
-		FUIAction(
-			FExecuteAction::CreateRaw(this, &FLevelSequenceCustomization::SetSelectedNodesSpawnableLevel, ObjectBindingModel, FName(NAME_None)),
-			FCanExecuteAction(),
-			FIsActionChecked::CreateLambda([=] { return Spawnable->GetLevelName() == NAME_None; })
-		),
-		NAME_None,
-		EUserInterfaceActionType::ToggleButton
-	);
-
-	UWorld* World = Sequencer->GetPlaybackContext()->GetWorld();
-	if (!World)
-	{
-		return;
-	}
-
-	for (ULevelStreaming* LevelStreaming : World->GetStreamingLevels())
-	{
-		if (LevelStreaming)
-		{
-			FName LevelName = FPackageName::GetShortFName(LevelStreaming->GetWorldAssetPackageFName());
-
-			MenuBuilder.AddMenuEntry(
-				FText::FromName(LevelName),
-				FText::FromName(LevelName),
-				FSlateIcon(),
-				FUIAction(
-					FExecuteAction::CreateRaw(this, &FLevelSequenceCustomization::SetSelectedNodesSpawnableLevel, ObjectBindingModel, LevelName),
-					FCanExecuteAction(),
-					FIsActionChecked::CreateLambda([=] { return Spawnable->GetLevelName() == LevelName; })
-				),
-				NAME_None,
-				EUserInterfaceActionType::ToggleButton
-			);
-		}
-	}
-}
-
-void FLevelSequenceCustomization::SetSelectedNodesSpawnableLevel(TSharedPtr<FObjectBindingModel> ObjectBindingModel, FName InLevelName)
-{
-	TSharedPtr<ISequencer> Sequencer = WeakSequencer.Pin();
-	TSharedPtr<FSequencerEditorViewModel> EditorViewModel = Sequencer->GetViewModel();
-	UMovieScene* MovieScene = Sequencer->GetFocusedMovieSceneSequence()->GetMovieScene();
-
-	if (MovieScene->IsReadOnly())
-	{
-		FSequencerUtilities::ShowReadOnlyError();
-		return;
-	}
-
-	const FScopedTransaction Transaction( LOCTEXT("SetSpawnableLevel", "Set Spawnable Level") );
-
-	MovieScene->Modify();
-
-	TArray<FMovieSceneSpawnable*> Spawnables;
-
-	for (TViewModelPtr<IObjectBindingExtension> ObjectBindingNode : EditorViewModel->GetSelection()->Outliner.Filter<IObjectBindingExtension>())
-	{
-		FMovieSceneSpawnable* Spawnable = MovieScene->FindSpawnable(ObjectBindingNode->GetObjectGuid());
-		if (Spawnable)
-		{
-			Spawnable->SetLevelName(InLevelName);
-		}
-	}
-}
-
-void FLevelSequenceCustomization::AddDynamicSpawnMenu(FMenuBuilder& MenuBuilder, TSharedPtr<FObjectBindingModel> ObjectBindingModel)
-{
-	TSharedPtr<ISequencer> Sequencer = WeakSequencer.Pin();
-	UMovieScene* MovieScene = Sequencer->GetFocusedMovieSceneSequence()->GetMovieScene();
-
-	FGuid ObjectBindingID = ObjectBindingModel->GetObjectGuid();
-	FMovieSceneSpawnable* Spawnable = MovieScene->FindSpawnable(ObjectBindingID);
-	if (!Spawnable)
-	{
-		return;
-	}
-
-	ObjectBindingModel->AddDynamicBindingMenu(MenuBuilder, Spawnable->DynamicBinding);
-}
-
-void FLevelSequenceCustomization::AddDynamicPossessionMenu(FMenuBuilder& MenuBuilder, TSharedPtr<FObjectBindingModel> ObjectBindingModel)
-{
-	TSharedPtr<ISequencer> Sequencer = WeakSequencer.Pin();
-	UMovieScene* MovieScene = Sequencer->GetFocusedMovieSceneSequence()->GetMovieScene();
-
-	FGuid ObjectBindingID = ObjectBindingModel->GetObjectGuid();
-	FMovieScenePossessable* Possessable = MovieScene->FindPossessable(ObjectBindingID);
-	if (!Possessable)
-	{
-		return;
-	}
-
-	ObjectBindingModel->AddDynamicBindingMenu(MenuBuilder, Possessable->DynamicBinding);
-}
-
 } // namespace UE::Sequencer
 
 #undef LOCTEXT_NAMESPACE
diff --git a/Engine/Plugins/MovieScene/LevelSequenceEditor/Source/LevelSequenceEditor/Private/Misc/LevelSequenceCustomization.h b/Engine/Plugins/MovieScene/LevelSequenceEditor/Source/LevelSequenceEditor/Private/Misc/LevelSequenceCustomization.h
--- a/Engine/Plugins/MovieScene/LevelSequenceEditor/Source/LevelSequenceEditor/Private/Misc/LevelSequenceCustomization.h
+++ b/Engine/Plugins/MovieScene/LevelSequenceEditor/Source/LevelSequenceEditor/Private/Misc/LevelSequenceCustomization.h
@@ -34,11 +34,6 @@ private:
 	// Object binding context menu extensions
 	TSharedPtr<FExtender> CreateObjectBindingContextMenuExtender(FViewModelPtr InViewModel);
 	void ExtendObjectBindingContextMenu(FMenuBuilder& MenuBuilder, TSharedPtr<FObjectBindingModel> ObjectBindingModel);
-	void AddSpawnOwnershipMenu(FMenuBuilder& MenuBuilder, TSharedPtr<FObjectBindingModel> ObjectBindingModel);
-	void AddSpawnLevelMenu(FMenuBuilder& MenuBuilder, TSharedPtr<FObjectBindingModel> ObjectBindingModel);
-	void SetSelectedNodesSpawnableLevel(TSharedPtr<FObjectBindingModel> ObjectBindingModel, FName InLevelName);
-	void AddDynamicSpawnMenu(FMenuBuilder& MenuBuilder, TSharedPtr<FObjectBindingModel> ObjectBindingModel);
-	void AddDynamicPossessionMenu(FMenuBuilder& MenuBuilder, TSharedPtr<FObjectBindingModel> ObjectBindingModel);
 
 private:
 
diff --git a/Engine/Plugins/MovieScene/LevelSequenceEditor/Source/LevelSequenceEditor/Private/Misc/LevelSequenceEditorActorBinding.cpp b/Engine/Plugins/MovieScene/LevelSequenceEditor/Source/LevelSequenceEditor/Private/Misc/LevelSequenceEditorActorBinding.cpp
--- a/Engine/Plugins/MovieScene/LevelSequenceEditor/Source/LevelSequenceEditor/Private/Misc/LevelSequenceEditorActorBinding.cpp
+++ b/Engine/Plugins/MovieScene/LevelSequenceEditor/Source/LevelSequenceEditor/Private/Misc/LevelSequenceEditorActorBinding.cpp
@@ -61,7 +61,9 @@ void FLevelSequenceEditorActorBinding::AddPossessActorMenuExtensions(FMenuBuilde
 				if (!Possessable.GetParent().IsValid())
 				{
 					// A possession guid can apply to more than one object, so we get all bound objects for the GUID and add them to our set.
-					ExistingPossessedObjects.Append(MovieSceneSequence->LocateBoundObjects(Possessable.GetGuid(), Sequencer.Pin()->GetPlaybackContext()));
+					TArray<UObject*, TInlineAllocator<1>> OutObjects;
+					MovieSceneSequence->LocateBoundObjects(Possessable.GetGuid(), UE::UniversalObjectLocator::FResolveParams(Sequencer.Pin()->GetPlaybackContext()), Sequencer.Pin()->FindSharedPlaybackState(), OutObjects);
+					ExistingPossessedObjects.Append(OutObjects);
 				}
 			}
 		}
diff --git a/Engine/Plugins/MovieScene/LevelSequenceEditor/Source/LevelSequenceEditor/Public/LevelSequenceEditorSubsystem.h b/Engine/Plugins/MovieScene/LevelSequenceEditor/Source/LevelSequenceEditor/Public/LevelSequenceEditorSubsystem.h
--- a/Engine/Plugins/MovieScene/LevelSequenceEditor/Source/LevelSequenceEditor/Public/LevelSequenceEditorSubsystem.h
+++ b/Engine/Plugins/MovieScene/LevelSequenceEditor/Source/LevelSequenceEditor/Public/LevelSequenceEditorSubsystem.h
@@ -36,6 +36,7 @@ class USequencerModuleScriptingLayer;
 class IDetailsView;
 class USequencerCurveEditorObject;
 class UMovieSceneCustomBinding;
+class IMenu;
 
 USTRUCT(BlueprintType)
 struct FMovieSceneScriptingParams
@@ -256,12 +257,15 @@ private:
 		virtual void NotifyPostChange(const FPropertyChangedEvent& PropertyChangedEvent, FProperty* PropertyThatChanged) override;
 	};
 
+	UPROPERTY()
 	TObjectPtr<UMovieSceneBindingPropertyInfoList> BindingPropertyInfoList = nullptr;
+	
 	FBindingPropertiesNotifyHook NotifyHook;
 
 private:
 
 	void AddBindingDetailCustomizations(TSharedRef<IDetailsView> DetailsView, TSharedPtr<ISequencer> ActiveSequencer, FGuid BindingGuid);
+	void OnMenuBeingDestroyed(const TSharedRef<IMenu>& Menu, TSharedRef<IDetailsView> DetailsView);
 
 	TSharedPtr<ISequencer> GetActiveSequencer();
 	
diff --git a/Engine/Plugins/MovieScene/MovieRenderPipeline/Source/MovieRenderPipelineCore/Private/MoviePipelineBlueprintLibrary.cpp b/Engine/Plugins/MovieScene/MovieRenderPipeline/Source/MovieRenderPipelineCore/Private/MoviePipelineBlueprintLibrary.cpp
--- a/Engine/Plugins/MovieScene/MovieRenderPipeline/Source/MovieRenderPipelineCore/Private/MoviePipelineBlueprintLibrary.cpp
+++ b/Engine/Plugins/MovieScene/MovieRenderPipeline/Source/MovieRenderPipelineCore/Private/MoviePipelineBlueprintLibrary.cpp
@@ -6,7 +6,6 @@
 #include "MovieRenderPipelineDataTypes.h"
 #include "MovieScene.h"
 #include "MovieSceneSequence.h"
-#include "MovieSceneSpawnable.h"
 #include "MovieScenePossessable.h"
 #include "MovieSceneBinding.h"
 #include "LevelSequence.h"
@@ -512,6 +511,8 @@ void UMoviePipelineBlueprintLibrary::UpdateJobShotListFromSequence(ULevelSequenc
 			// camera components and track them as sidecar cameras. We do this even if sidecar data isn't being used.
 			if (LeafNode->CameraCutSection.IsValid() && LeafNode->MovieScene.IsValid())
 			{
+				UMovieSceneSequence* Sequence = LeafNode->MovieScene->GetTypedOuter<UMovieSceneSequence>();
+				TSharedRef<UE::MovieScene::FSharedPlaybackState> TransientPlaybackState = MovieSceneHelpers::CreateTransientSharedPlaybackState(GWorld, Sequence);
 				int32 MainCameraIndex = INDEX_NONE;
 				FMovieSceneObjectBindingID MainBinding = LeafNode->CameraCutSection->GetCameraBindingID();
 
@@ -561,33 +562,28 @@ void UMoviePipelineBlueprintLibrary::UpdateJobShotListFromSequence(ULevelSequenc
 								LastValidGuid = ParentGuid;
 								ParentGuid = ParentAsPossessable->GetParent();
 								SidecarCamera.Name = ParentAsPossessable->GetName();
-							}
-							else if (FMovieSceneSpawnable* ParentAsSpawnable = LeafNode->MovieScene->FindSpawnable(ParentGuid))
-							{
-								FMovieSceneBinding* Binding = LeafNode->MovieScene->FindBinding(ParentAsSpawnable->GetGuid());
-								if (Binding)
+
+								if (MovieSceneHelpers::IsBoundToAnySpawnable(LeafNode->MovieScene->GetTypedOuter<UMovieSceneSequence>(), ParentAsPossessable->GetGuid(), TransientPlaybackState))
 								{
-									const TArray<UMovieSceneTrack*> AllBindingTracks = Binding->GetTracks();
-									for (const UMovieSceneTrack* Track : AllBindingTracks)
+									FMovieSceneBinding* Binding = LeafNode->MovieScene->FindBinding(ParentAsPossessable->GetGuid());
+									if (Binding)
 									{
-										if (const UMovieSceneSpawnTrack* SpawnTrack = Cast<UMovieSceneSpawnTrack>(Track))
+										const TArray<UMovieSceneTrack*> AllBindingTracks = Binding->GetTracks();
+										for (const UMovieSceneTrack* Track : AllBindingTracks)
 										{
-											for (const UMovieSceneSection* Section : SpawnTrack->GetAllSections())
+											if (const UMovieSceneSpawnTrack* SpawnTrack = Cast<UMovieSceneSpawnTrack>(Track))
 											{
-												if (const UMovieSceneSpawnSection* BoolSection = Cast<UMovieSceneSpawnSection>(Section))
+												for (const UMovieSceneSection* Section : SpawnTrack->GetAllSections())
 												{
-													bAddCamera = BoolSection->GetChannel().GetDefault().Get(true);
+													if (const UMovieSceneSpawnSection* BoolSection = Cast<UMovieSceneSpawnSection>(Section))
+													{
+														bAddCamera = BoolSection->GetChannel().GetDefault().Get(true);
+													}
 												}
 											}
 										}
 									}
 								}
-
-								// Spawnables will never have a parent
-								LastValidGuid = ParentGuid;
-								ParentGuid.Invalidate(); 
-								
-								SidecarCamera.Name = ParentAsSpawnable->GetName();
 							}
 						}
 
diff --git a/Engine/Plugins/MovieScene/SequencerScripting/Source/SequencerScripting/Private/ExtensionLibraries/MovieSceneBindingExtensions.cpp b/Engine/Plugins/MovieScene/SequencerScripting/Source/SequencerScripting/Private/ExtensionLibraries/MovieSceneBindingExtensions.cpp
--- a/Engine/Plugins/MovieScene/SequencerScripting/Source/SequencerScripting/Private/ExtensionLibraries/MovieSceneBindingExtensions.cpp
+++ b/Engine/Plugins/MovieScene/SequencerScripting/Source/SequencerScripting/Private/ExtensionLibraries/MovieSceneBindingExtensions.cpp
@@ -260,16 +260,9 @@ UObject* UMovieSceneBindingExtensions::GetObjectTemplate(const FMovieSceneBindin
 	UMovieScene* MovieScene = InBinding.GetMovieScene();
 	if (MovieScene)
 	{
-		using namespace UE::MovieScene;
-		FSharedPlaybackStateCreateParams CreateParams;
-		// Ideally we would have a proper world context here, but we don't so use GWorld
-		CreateParams.PlaybackContext = GWorld;
 		UMovieSceneSequence* ThisSequence = MovieScene->GetTypedOuter<UMovieSceneSequence>();
-		TSharedRef<FSharedPlaybackState> TransientPlaybackState = MakeShared<FSharedPlaybackState>(*ThisSequence, CreateParams);
+		TSharedRef<UE::MovieScene::FSharedPlaybackState> TransientPlaybackState = MovieSceneHelpers::CreateTransientSharedPlaybackState(GWorld, ThisSequence);
 
-		FMovieSceneEvaluationState State;
-		TransientPlaybackState->AddCapabilityRaw(&State);
-		State.AssignSequence(MovieSceneSequenceID::Root, *ThisSequence, TransientPlaybackState);
 
 		// TODO: Technically this assumes only one spawnable- do we need to upgrade script to handle multiple binding indices?
 		return MovieSceneHelpers::GetObjectTemplate(MovieScene->GetTypedOuter<UMovieSceneSequence>(), InBinding.BindingID, TransientPlaybackState);
diff --git a/Engine/Plugins/VirtualProduction/Takes/Source/TakeRecorder/Private/SerializedRecorder.cpp b/Engine/Plugins/VirtualProduction/Takes/Source/TakeRecorder/Private/SerializedRecorder.cpp
--- a/Engine/Plugins/VirtualProduction/Takes/Source/TakeRecorder/Private/SerializedRecorder.cpp
+++ b/Engine/Plugins/VirtualProduction/Takes/Source/TakeRecorder/Private/SerializedRecorder.cpp
@@ -32,6 +32,7 @@
 #include "UObject/UObjectIterator.h"
 #include "LevelSequenceActor.h"
 #include "LevelEditorViewport.h"
+#include "SequencerUtilities.h"
 
 #define LOCTEXT_NAMESPACE "SerializedRecorder"
 
@@ -423,11 +424,18 @@ AActor* FSerializedRecorder::SetActorPossesableOrSpawnable(UMovieSceneSequence*
 
 		Actor->SetActorLabel(ActorHeader.Label); //spawnable has samen Name used for matching.
 
-		FMovieSceneSpawnable Spawnable(ActorHeader.Label, *Actor);
-		Spawnable.SetGuid(ActorHeader.Guid);
+		FMovieScenePossessable Possessable(ActorHeader.Label, SpawnableClass);
 
+		Possessable.SetGuid(ActorHeader.Guid);
 		FMovieSceneBinding NewBinding(ActorHeader.Guid, ActorHeader.Label);
-		MovieScene->AddSpawnable(Spawnable, NewBinding);
+		MovieScene->AddPossessable(Possessable, NewBinding);
+
+		UE::Sequencer::FCreateBindingParams CreateBindingParams;
+		CreateBindingParams.bAllowCustomBinding = true;
+		CreateBindingParams.bSpawnable = true;
+		CreateBindingParams.BindingNameOverride = ActorHeader.Label;
+		CreateBindingParams.ReplacementGuid = ActorHeader.Guid;
+		FSequencerUtilities::CreateOrReplaceBinding(nullptr, InMovieSceneSequence, Actor, CreateBindingParams);
 
 		if (ActorHeader.Guid.IsValid())
 		{
@@ -475,15 +483,6 @@ AActor* FSerializedRecorder::SetActorPossesableOrSpawnable(UMovieSceneSequence*
 		Possessable->SetName(ActorHeader.Label);
 	}
 
-	if (FMovieSceneSpawnable* Spawnable = MovieScene->FindSpawnable(ActorHeader.Guid))
-	{
-		if (!Spawnable->Tags.Contains(*(ActorHeader.Label)))
-		{
-			Spawnable->Tags.AddUnique(FName(*(ActorHeader.Label)));
-		}
-		Spawnable->SetName(ActorHeader.Label);
-	}
-
 	// look for a folder to put us in
 	if (ActorHeader.FolderName.IsValid())
 	{
@@ -526,12 +525,6 @@ void FSerializedRecorder::SetComponentPossessable(UMovieSceneSequence* InMovieSc
 		// Set up parent/child guids for possessables within spawnables
 		FMovieScenePossessable* ChildPossessablePtr = InMovieScene->FindPossessable(ActorProperty.Guid);
 		ChildPossessablePtr->SetParent(ActorHeader.Guid, InMovieScene);
-		
-		FMovieSceneSpawnable* ParentSpawnable = InMovieScene->FindSpawnable(ActorHeader.Guid);
-		if (ParentSpawnable)
-		{
-			ParentSpawnable->AddChildPossessable(ActorProperty.Guid);
-		}
 
 		for (UActorComponent* Component : Actor->GetComponents())
 		{
