Sequencer:
Fix shutdown issue found in cook calling FSlateApplication::Get.

diff --git a/Engine/Plugins/MovieScene/LevelSequenceEditor/Source/LevelSequenceEditor/Private/LevelSequenceEditorSubsystem.cpp b/Engine/Plugins/MovieScene/LevelSequenceEditor/Source/LevelSequenceEditor/Private/LevelSequenceEditorSubsystem.cpp
--- a/Engine/Plugins/MovieScene/LevelSequenceEditor/Source/LevelSequenceEditor/Private/LevelSequenceEditorSubsystem.cpp
+++ b/Engine/Plugins/MovieScene/LevelSequenceEditor/Source/LevelSequenceEditor/Private/LevelSequenceEditorSubsystem.cpp
@@ -507,7 +507,10 @@ void ULevelSequenceEditorSubsystem::Deinitialize()
 		SequencerModulePtr->UnregisterOnSequencerCreated(OnSequencerCreatedHandle);
 	}
 
-	FSlateApplication::Get().OnMenuBeingDestroyed().RemoveAll(this);
+	if (FSlateApplication::IsInitialized())
+	{
+		FSlateApplication::Get().OnMenuBeingDestroyed().RemoveAll(this);
+	}
 
 }
 
@@ -2197,8 +2200,10 @@ void ULevelSequenceEditorSubsystem::AddBindingPropertiesMenu(FMenuBuilder& MenuB
 		RefreshBindingDetails(&DetailsView.Get(), ObjectBindings[0]);
 		DetailsView->OnFinishedChangingProperties().AddUObject(this, &ULevelSequenceEditorSubsystem::OnFinishedChangingLocators, DetailsView, ObjectBindings[0]);
 
-		FSlateApplication::Get().OnMenuBeingDestroyed().AddUObject(this, &ULevelSequenceEditorSubsystem::OnMenuBeingDestroyed, DetailsView);
-
+		if (FSlateApplication::IsInitialized())
+		{
+			FSlateApplication::Get().OnMenuBeingDestroyed().AddUObject(this, &ULevelSequenceEditorSubsystem::OnMenuBeingDestroyed, DetailsView);
+		}
 		MenuBuilder.AddWidget(DetailsView, FText::GetEmpty(), true);
 	}
 }
