Sequencer: Expose custom binding logic classes to blueprint/script and fixup some of the script functions to work better with custo bindings.

diff --git a/Engine/Plugins/MovieScene/LevelSequenceEditor/Source/LevelSequenceEditor/Private/LevelSequenceEditorSubsystem.cpp b/Engine/Plugins/MovieScene/LevelSequenceEditor/Source/LevelSequenceEditor/Private/LevelSequenceEditorSubsystem.cpp
--- a/Engine/Plugins/MovieScene/LevelSequenceEditor/Source/LevelSequenceEditor/Private/LevelSequenceEditorSubsystem.cpp
+++ b/Engine/Plugins/MovieScene/LevelSequenceEditor/Source/LevelSequenceEditor/Private/LevelSequenceEditorSubsystem.cpp
@@ -723,14 +723,96 @@ FMovieSceneBindingProxy ULevelSequenceEditorSubsystem::ConvertToPossessable(cons
 		return PossessableProxy;
 	}
 
-	if (FMovieScenePossessable* Possessable = FSequencerUtilities::ConvertToPossessable(Sequencer.ToSharedRef(), ObjectBinding.BindingID))
+	FMovieScenePossessable* NewPossessable = nullptr;
+
+	if (FMovieSceneBindingReferences* BindingReferences = Sequence->GetBindingReferences())
 	{
-		PossessableProxy = FMovieSceneBindingProxy(Possessable->GetGuid(), Sequence);
+		int32 NumBindings = BindingReferences->GetReferences(ObjectBinding.BindingID).Num();
+		for (int32 BindingIndex = 0; BindingIndex < NumBindings; BindingIndex++)
+		{
+			NewPossessable = FSequencerUtilities::ConvertToPossessable(Sequencer.ToSharedRef(), ObjectBinding.BindingID, BindingIndex);
+		}
+	}
+
+	if (NewPossessable)
+	{
+		PossessableProxy = FMovieSceneBindingProxy(NewPossessable->GetGuid(), Sequence);
 	}
 
 	return PossessableProxy;
 }
 
+
+FMovieSceneBindingProxy ULevelSequenceEditorSubsystem::ConvertToCustomBinding(const FMovieSceneBindingProxy& ObjectBinding, TSubclassOf<UMovieSceneCustomBinding> BindingType)
+{
+	FMovieSceneBindingProxy PossessableProxy;
+
+	TSharedPtr<ISequencer> Sequencer = GetActiveSequencer();
+	if (Sequencer == nullptr)
+	{
+		return PossessableProxy;
+	}
+
+	UMovieSceneSequence* Sequence = Sequencer->GetFocusedMovieSceneSequence();
+	if (!Sequence)
+	{
+		return PossessableProxy;
+	}
+
+	FMovieSceneBindingReferences* BindingReferences = Sequence->GetBindingReferences();
+	if (!BindingReferences)
+	{
+		return PossessableProxy;
+	}
+
+	FMovieScenePossessable* NewPossessable = nullptr;
+
+	if (FSequencerUtilities::CanConvertToCustomBinding(Sequencer.ToSharedRef(), ObjectBinding.BindingID, BindingType))
+	{
+		int32 NumBindings = BindingReferences->GetReferences(ObjectBinding.BindingID).Num();
+		for (int32 BindingIndex = 0; BindingIndex < NumBindings; BindingIndex++)
+		{
+			NewPossessable = FSequencerUtilities::ConvertToCustomBinding(Sequencer.ToSharedRef(), ObjectBinding.BindingID, BindingType, BindingIndex);
+		}
+
+		if (NewPossessable)
+		{
+			PossessableProxy = FMovieSceneBindingProxy(NewPossessable->GetGuid(), Sequence);
+		}
+	}
+
+	return PossessableProxy;
+}
+
+TArray<UMovieSceneCustomBinding*> ULevelSequenceEditorSubsystem::GetCustomBindingObjects(const FMovieSceneBindingProxy& ObjectBinding)
+{
+	TArray<UMovieSceneCustomBinding*> CustomBindings;
+	TSharedPtr<ISequencer> Sequencer = GetActiveSequencer();
+	if (Sequencer == nullptr)
+	{
+		return CustomBindings;
+	}
+
+	UMovieSceneSequence* Sequence = Sequencer->GetFocusedMovieSceneSequence();
+	if (!Sequence)
+	{
+		return CustomBindings;
+	}
+
+	if (FMovieSceneBindingReferences* BindingReferences = Sequence->GetBindingReferences())
+	{
+		for (const FMovieSceneBindingReference& BindingReference : BindingReferences->GetReferences(ObjectBinding.BindingID))
+		{
+			if (BindingReference.CustomBinding)
+			{
+				CustomBindings.Add(BindingReference.CustomBinding);
+			}
+		}
+	}
+
+	return CustomBindings;
+}
+
 void ULevelSequenceEditorSubsystem::CopyFolders(const TArray<UMovieSceneFolder*>& Folders, FString& ExportedText)
 {
 	FSequencerUtilities::CopyFolders(Folders, ExportedText);
diff --git a/Engine/Plugins/MovieScene/LevelSequenceEditor/Source/LevelSequenceEditor/Public/LevelSequenceEditorSubsystem.h b/Engine/Plugins/MovieScene/LevelSequenceEditor/Source/LevelSequenceEditor/Public/LevelSequenceEditorSubsystem.h
--- a/Engine/Plugins/MovieScene/LevelSequenceEditor/Source/LevelSequenceEditor/Public/LevelSequenceEditorSubsystem.h
+++ b/Engine/Plugins/MovieScene/LevelSequenceEditor/Source/LevelSequenceEditor/Public/LevelSequenceEditorSubsystem.h
@@ -110,14 +110,25 @@ public:
 	UFUNCTION(BlueprintCallable, Category = "Level Sequence Editor")
 	FMovieSceneBindingProxy CreateCamera(bool bSpawnable, ACineCameraActor*& OutActor);
 
-	/** Convert to spawnable. If there are multiple objects assigned to the possessable, multiple spawnables will be created. */
+	/** 
+	* Convert to spawnable. If there are multiple objects assigned to the possessable, multiple spawnables will be created. 
+	* For level sequences, the bindings created will be custom bindings of type UMovieSceneSpawnableActorBinding.
+	*/
 	UFUNCTION(BlueprintCallable, Category = "Level Sequence Editor")
 	TArray<FMovieSceneBindingProxy> ConvertToSpawnable(const FMovieSceneBindingProxy& ObjectBinding);
 
-	/** Convert to possessable */
+	/** Convert to possessable. If there are multiple objects assigned to the spawnable. */
 	UFUNCTION(BlueprintCallable, Category = "Level Sequence Editor")
 	FMovieSceneBindingProxy ConvertToPossessable(const FMovieSceneBindingProxy& ObjectBinding);
 
+	/** Convert to a custom binding of the given binding type*/
+	UFUNCTION(BlueprintCallable, Category = "Level Sequence Editor")
+	FMovieSceneBindingProxy ConvertToCustomBinding(const FMovieSceneBindingProxy& ObjectBinding, UPARAM(meta = (AllowAbstract = "false")) TSubclassOf<UMovieSceneCustomBinding> BindingType);
+
+	/** In the case that the given binding proxy holds custom bindings, returns an array of the binding objects so properties can be accessed. */
+	UFUNCTION(BlueprintCallable, Category = "Level Sequence Editor")
+	TArray<UMovieSceneCustomBinding*> GetCustomBindingObjects(const FMovieSceneBindingProxy& ObjectBinding);
+
 	/** 
 	 * Copy folders 
 	 * The copied folders will be saved to the clipboard as well as assigned to the ExportedText string. 
diff --git a/Engine/Plugins/MovieScene/SequencerScripting/Source/SequencerScripting/Private/ExtensionLibraries/MovieSceneSequenceExtensions.cpp b/Engine/Plugins/MovieScene/SequencerScripting/Source/SequencerScripting/Private/ExtensionLibraries/MovieSceneSequenceExtensions.cpp
--- a/Engine/Plugins/MovieScene/SequencerScripting/Source/SequencerScripting/Private/ExtensionLibraries/MovieSceneSequenceExtensions.cpp
+++ b/Engine/Plugins/MovieScene/SequencerScripting/Source/SequencerScripting/Private/ExtensionLibraries/MovieSceneSequenceExtensions.cpp
@@ -11,6 +11,8 @@
 #include "Compilation/MovieSceneCompiledDataManager.h"
 #include "MovieSceneSpawnable.h"
 #include "SequencerScriptingRange.h"
+#include "MovieSceneBindingReferences.h"
+#include "Bindings/MovieSceneSpawnableActorBinding.h"
 
 #include UE_INLINE_GENERATED_CPP_BY_NAME(MovieSceneSequenceExtensions)
 
@@ -732,6 +734,17 @@ TArray<FMovieSceneBindingProxy> UMovieSceneSequenceExtensions::GetSpawnables(UMo
 		}
 	}
 
+	if (const FMovieSceneBindingReferences* BindingReferences = Sequence->GetBindingReferences())
+	{
+		for (const FMovieSceneBindingReference& BindingReference : BindingReferences->GetAllReferences())
+		{
+			if (BindingReference.CustomBinding && BindingReference.CustomBinding->IsA<UMovieSceneSpawnableActorBinding>())
+			{
+				AllSpawnables.Emplace(BindingReference.ID, Sequence);
+			}
+		}
+	}
+
 	return AllSpawnables;
 }
 
@@ -745,14 +758,26 @@ TArray<FMovieSceneBindingProxy> UMovieSceneSequenceExtensions::GetPossessables(U
 
 	TArray<FMovieSceneBindingProxy> AllPossessables;
 
-	UMovieScene* MovieScene = GetMovieScene(Sequence);
-	if (MovieScene)
+	if (const FMovieSceneBindingReferences* BindingReferences = Sequence->GetBindingReferences())
 	{
-		int32 Count = MovieScene->GetPossessableCount();
-		AllPossessables.Reserve(Count);
-		for (int32 i = 0; i < Count; ++i)
+		for (const FMovieSceneBindingReference& BindingReference : BindingReferences->GetAllReferences())
 		{
-			AllPossessables.Emplace(MovieScene->GetPossessable(i).GetGuid(), Sequence);
+			if (!BindingReference.CustomBinding)
+			{
+				AllPossessables.Emplace(BindingReference.ID, Sequence);
+			}
+		}
+	}
+	else
+	{
+		UMovieScene* MovieScene = GetMovieScene(Sequence);
+		if (MovieScene)
+		{
+			int32 Count = MovieScene->GetPossessableCount();
+			for (int32 i = 0; i < Count; ++i)
+			{
+				AllPossessables.Emplace(MovieScene->GetPossessable(i).GetGuid(), Sequence);
+			}
 		}
 	}
 
diff --git a/Engine/Plugins/MovieScene/SequencerScripting/Source/SequencerScripting/Public/ExtensionLibraries/MovieSceneSequenceExtensions.h b/Engine/Plugins/MovieScene/SequencerScripting/Source/SequencerScripting/Public/ExtensionLibraries/MovieSceneSequenceExtensions.h
--- a/Engine/Plugins/MovieScene/SequencerScripting/Source/SequencerScripting/Public/ExtensionLibraries/MovieSceneSequenceExtensions.h
+++ b/Engine/Plugins/MovieScene/SequencerScripting/Source/SequencerScripting/Public/ExtensionLibraries/MovieSceneSequenceExtensions.h
@@ -372,7 +372,7 @@ public:
 	static TArray<FMovieSceneBindingProxy> GetBindings(UMovieSceneSequence* Sequence);
 
 	/**
-	* Get all the spawnables in this sequence
+	* Get all the spawnables in this sequence. For Level Sequences, this includes bindings with binding type UMovieSceneSpawnableActorBinding.
 	*
 	* @param Sequence        The sequence to get spawnables for
 	* @return Spawnables in this sequence
@@ -381,7 +381,7 @@ public:
 	static TArray<FMovieSceneBindingProxy> GetSpawnables(UMovieSceneSequence* Sequence);
 
 	/**
-	* Get all the possessables in this sequence
+	* Get all the possessables in this sequence. It is understood for the purpose of this function that this means the bindings are not custom.
 	*
 	* @param Sequence        The sequence to get possessables for
 	* @return Possessables in this sequence
@@ -400,7 +400,8 @@ public:
 	static FMovieSceneBindingProxy AddPossessable(UMovieSceneSequence* Sequence, UObject* ObjectToPossess);
 
 	/**
-	 * Add a new binding to this sequence that will spawn the specified object
+	 * Add a new binding to this sequence that will spawn the specified object.
+	 * For level sequences this will make a custom binding of type UMovieSceneSpawnableActorBinding.
 	 *
 	 * @param Sequence        The sequence to add to
 	 * @param ObjectToSpawn   An object instance to use as a template for spawning
@@ -411,6 +412,7 @@ public:
 
 	/**
 	 * Add a new binding to this sequence that will spawn the specified object
+	 * For level sequences this will make a custom binding of type UMovieSceneSpawnableActorBinding.
 	 *
 	 * @param Sequence        The sequence to add to
 	 * @param ClassToSpawn    A class or blueprint type to spawn for this binding
diff --git a/Engine/Source/Runtime/MovieScene/Public/Bindings/MovieSceneReplaceableBinding.h b/Engine/Source/Runtime/MovieScene/Public/Bindings/MovieSceneReplaceableBinding.h
--- a/Engine/Source/Runtime/MovieScene/Public/Bindings/MovieSceneReplaceableBinding.h
+++ b/Engine/Source/Runtime/MovieScene/Public/Bindings/MovieSceneReplaceableBinding.h
@@ -51,7 +51,7 @@ public:
 
 #if WITH_EDITORONLY_DATA
 	// Optional Editor-only preview object
-	UPROPERTY(Instanced, VisibleAnywhere, Category="Editor")
+	UPROPERTY(Instanced, VisibleAnywhere, BlueprintReadOnly, Category="Editor")
 	TObjectPtr<UMovieSceneSpawnableBindingBase> PreviewSpawnable = nullptr;
 #endif
 
diff --git a/Engine/Source/Runtime/MovieScene/Public/Bindings/MovieSceneSpawnableBinding.h b/Engine/Source/Runtime/MovieScene/Public/Bindings/MovieSceneSpawnableBinding.h
--- a/Engine/Source/Runtime/MovieScene/Public/Bindings/MovieSceneSpawnableBinding.h
+++ b/Engine/Source/Runtime/MovieScene/Public/Bindings/MovieSceneSpawnableBinding.h
@@ -67,11 +67,11 @@ public:
 	/*
 	* The spawn ownership setting for this spawnable, allowing spawnables to potentially outlast the lifetime of their sub sequence or sequence altogether.
 	*/
-	UPROPERTY(EditAnywhere, Category=Spawnable)
+	UPROPERTY(EditAnywhere, BlueprintReadWrite, Category=Spawnable)
 	ESpawnOwnership SpawnOwnership = ESpawnOwnership::InnerSequence;
 
 	/** When enabled, this spawnable will always be respawned if it gets destroyed externally. When disabled, this object will only ever be spawned once for each binding lifetime section even if destroyed externally. */
-	UPROPERTY(EditAnywhere, Category= Spawnable)
+	UPROPERTY(EditAnywhere, BlueprintReadWrite, Category= Spawnable)
 	bool bContinuouslyRespawn = false;
 
 public:
diff --git a/Engine/Source/Runtime/MovieSceneTracks/Public/Bindings/MovieSceneReplaceableDirectorBlueprintBinding.h b/Engine/Source/Runtime/MovieSceneTracks/Public/Bindings/MovieSceneReplaceableDirectorBlueprintBinding.h
--- a/Engine/Source/Runtime/MovieSceneTracks/Public/Bindings/MovieSceneReplaceableDirectorBlueprintBinding.h
+++ b/Engine/Source/Runtime/MovieSceneTracks/Public/Bindings/MovieSceneReplaceableDirectorBlueprintBinding.h
@@ -37,7 +37,7 @@ public:
 	// Preview Spawnable Type to use for this replaceable
 	// TODO: Do I want to consider moving this sort of thing to the base class so any replaceable type can choose any other spawnable type?
 	// TODO: Changing this will require 
-	UPROPERTY(EditAnywhere, Category="Sequencer")
+	UPROPERTY(EditAnywhere, BlueprintReadWrite, Category="Sequencer")
 	TSubclassOf<UMovieSceneSpawnableBindingBase> PreviewSpawnableType;
 
 protected:
diff --git a/Engine/Source/Runtime/MovieSceneTracks/Public/Bindings/MovieSceneSpawnableActorBinding.h b/Engine/Source/Runtime/MovieSceneTracks/Public/Bindings/MovieSceneSpawnableActorBinding.h
--- a/Engine/Source/Runtime/MovieSceneTracks/Public/Bindings/MovieSceneSpawnableActorBinding.h
+++ b/Engine/Source/Runtime/MovieSceneTracks/Public/Bindings/MovieSceneSpawnableActorBinding.h
@@ -64,11 +64,11 @@ public:
 
 public:
 	/** When enabled, the actor will be spawned with a unique name so that it can be addressable between clients and servers. */
-	UPROPERTY(EditAnywhere, Category=Actor)
+	UPROPERTY(EditAnywhere, BlueprintReadWrite, Category=Spawnable)
 	bool bNetAddressableName = false;
 
 	/** Name of level to spawn into */
-	UPROPERTY(EditAnywhere, Category = Actor)
+	UPROPERTY(EditAnywhere, BlueprintReadWrite, Category = Spawnable)
 	FName LevelName;
 
 protected:
@@ -148,10 +148,8 @@ protected:
 	MOVIESCENETRACKS_API TSubclassOf<AActor> GetActorClass() const override;
 	MOVIESCENETRACKS_API AActor* GetActorTemplate() const override { return ActorTemplate; }
 
-
-private:
-
-	UPROPERTY()
+protected:
+	UPROPERTY(VisibleAnywhere, Category="Spawnable")
 	TObjectPtr<AActor> ActorTemplate;
 
 
