Sequencer: Adding some additional scripting functions for handling custom bindings.

diff --git a/Engine/Plugins/MovieScene/LevelSequenceEditor/Source/LevelSequenceEditor/Private/LevelSequenceEditorSubsystem.cpp b/Engine/Plugins/MovieScene/LevelSequenceEditor/Source/LevelSequenceEditor/Private/LevelSequenceEditorSubsystem.cpp
--- a/Engine/Plugins/MovieScene/LevelSequenceEditor/Source/LevelSequenceEditor/Private/LevelSequenceEditorSubsystem.cpp
+++ b/Engine/Plugins/MovieScene/LevelSequenceEditor/Source/LevelSequenceEditor/Private/LevelSequenceEditorSubsystem.cpp
@@ -813,6 +813,81 @@ TArray<UMovieSceneCustomBinding*> ULevelSequenceEditorSubsystem::GetCustomBindin
 	return CustomBindings;
 }
 
+TArray<FMovieSceneBindingProxy> ULevelSequenceEditorSubsystem::GetCustomBindingsOfType(TSubclassOf<UMovieSceneCustomBinding> CustomBindingType)
+{
+	TArray<FMovieSceneBindingProxy> Bindings;
+	TSharedPtr<ISequencer> Sequencer = GetActiveSequencer();
+	if (Sequencer == nullptr)
+	{
+		return Bindings;
+	}
+
+	UMovieSceneSequence* Sequence = Sequencer->GetFocusedMovieSceneSequence();
+	if (!Sequence)
+	{
+		return Bindings;
+	}
+
+	if (FMovieSceneBindingReferences* BindingReferences = Sequence->GetBindingReferences())
+	{
+		for (const FMovieSceneBindingReference& BindingReference : BindingReferences->GetAllReferences())
+		{
+			if (BindingReference.CustomBinding && BindingReference.CustomBinding->IsA(CustomBindingType))
+			{
+				Bindings.AddUnique(FMovieSceneBindingProxy(BindingReference.ID, Sequence));
+			}
+		}
+	}
+
+	return Bindings;
+}
+
+TSubclassOf<UMovieSceneCustomBinding> ULevelSequenceEditorSubsystem::GetCustomBindingType(const FMovieSceneBindingProxy& ObjectBinding)
+{
+	TSharedPtr<ISequencer> Sequencer = GetActiveSequencer();
+	if (Sequencer == nullptr)
+	{
+		return nullptr;
+	}
+
+	UMovieSceneSequence* Sequence = Sequencer->GetFocusedMovieSceneSequence();
+	if (!Sequence)
+	{
+		return nullptr;
+	}
+
+	if (FMovieSceneBindingReferences* BindingReferences = Sequence->GetBindingReferences())
+	{
+		if (UMovieSceneCustomBinding* CustomBinding = BindingReferences->GetCustomBinding(ObjectBinding.BindingID, 0))
+		{
+			return CustomBinding->GetClass();
+		}
+	}
+
+	return nullptr;
+}
+
+bool ULevelSequenceEditorSubsystem::ChangeActorTemplateClass(const FMovieSceneBindingProxy& ObjectBinding, TSubclassOf<AActor> ActorClass)
+{
+	TSharedPtr<ISequencer> Sequencer = GetActiveSequencer();
+	if (Sequencer == nullptr)
+	{
+		return false;
+	}
+
+	UMovieSceneSequence* Sequence = Sequencer->GetFocusedMovieSceneSequence();
+	if (!Sequence)
+	{
+		return false;
+	}
+
+	bool bSuccess = false;
+
+	FSequencerUtilities::HandleTemplateActorClassPicked(ActorClass, Sequencer.ToSharedRef(), ObjectBinding.BindingID, 0, [&bSuccess](){bSuccess=true;});
+
+	return bSuccess;
+}
+
 void ULevelSequenceEditorSubsystem::CopyFolders(const TArray<UMovieSceneFolder*>& Folders, FString& ExportedText)
 {
 	FSequencerUtilities::CopyFolders(Folders, ExportedText);
diff --git a/Engine/Plugins/MovieScene/LevelSequenceEditor/Source/LevelSequenceEditor/Public/LevelSequenceEditorSubsystem.h b/Engine/Plugins/MovieScene/LevelSequenceEditor/Source/LevelSequenceEditor/Public/LevelSequenceEditorSubsystem.h
--- a/Engine/Plugins/MovieScene/LevelSequenceEditor/Source/LevelSequenceEditor/Public/LevelSequenceEditorSubsystem.h
+++ b/Engine/Plugins/MovieScene/LevelSequenceEditor/Source/LevelSequenceEditor/Public/LevelSequenceEditorSubsystem.h
@@ -129,6 +129,18 @@ public:
 	UFUNCTION(BlueprintCallable, Category = "Level Sequence Editor")
 	TArray<UMovieSceneCustomBinding*> GetCustomBindingObjects(const FMovieSceneBindingProxy& ObjectBinding);
 
+	/** Returns all of the bindings in the sequence of the given custom type. */
+	UFUNCTION(BlueprintCallable, Category = "Level Sequence Editor")
+	TArray<FMovieSceneBindingProxy> GetCustomBindingsOfType(TSubclassOf<UMovieSceneCustomBinding> CustomBindingType);
+
+	/* Returns the custom binding type for the given binding, or nullptr for possessables*/
+	UFUNCTION(BlueprintCallable, Category = "Level Sequence Editor")
+	TSubclassOf<UMovieSceneCustomBinding> GetCustomBindingType(const FMovieSceneBindingProxy& ObjectBinding);
+
+	/* Sets the actor class for the spawnable or replaceable template, in the case those binding types support templates. */
+	UFUNCTION(BlueprintCallable, Category = "Level Sequence Editor")
+	bool ChangeActorTemplateClass(const FMovieSceneBindingProxy& ObjectBinding, TSubclassOf<AActor> ActorClass);
+
 	/** 
 	 * Copy folders 
 	 * The copied folders will be saved to the clipboard as well as assigned to the ExportedText string. 
