Sequencer: Move conversion UX up to the top of the binding properties menu out of the individual elements.

diff --git a/Engine/Plugins/MovieScene/LevelSequenceEditor/Source/LevelSequenceEditor/Private/LevelSequenceEditorSubsystem.cpp b/Engine/Plugins/MovieScene/LevelSequenceEditor/Source/LevelSequenceEditor/Private/LevelSequenceEditorSubsystem.cpp
--- a/Engine/Plugins/MovieScene/LevelSequenceEditor/Source/LevelSequenceEditor/Private/LevelSequenceEditorSubsystem.cpp
+++ b/Engine/Plugins/MovieScene/LevelSequenceEditor/Source/LevelSequenceEditor/Private/LevelSequenceEditorSubsystem.cpp
@@ -75,6 +75,190 @@ DEFINE_LOG_CATEGORY(LogLevelSequenceEditor);
 
 #define LOCTEXT_NAMESPACE "LevelSequenceEditor"
 
+class FMovieSceneBindingPropertyInfoListCustomization : public IDetailCustomization
+{
+public:
+
+	FMovieSceneBindingPropertyInfoListCustomization(TWeakPtr<ISequencer> InSequencer, UMovieScene* InMovieScene, FGuid InBindingGuid, ULevelSequenceEditorSubsystem* InLevelSequenceEditorSubsystem);
+
+	static TSharedRef<IDetailCustomization> MakeInstance(TWeakPtr<ISequencer> InSequencer, UMovieScene* InMovieScene, FGuid InBindingGuid, ULevelSequenceEditorSubsystem* InLevelSequenceEditorSubsystem);
+	virtual void CustomizeDetails(IDetailLayoutBuilder& DetailBuilder) override;
+
+
+private:
+
+	TSharedRef<SWidget> OnGetConvertMenuContent(IDetailLayoutBuilder& DetailBuilder);
+
+	TWeakPtr<ISequencer> SequencerPtr;
+	TObjectPtr<UMovieScene> MovieScene = nullptr;
+	TObjectPtr<ULevelSequenceEditorSubsystem> LevelSequenceEditorSubsystem;
+	FGuid BindingGuid;
+	TObjectPtr<UMovieSceneBindingPropertyInfoList> BindingList;
+
+	TArray<TSharedPtr<FText>> BindingTypeNames;
+};
+
+FMovieSceneBindingPropertyInfoListCustomization::FMovieSceneBindingPropertyInfoListCustomization(TWeakPtr<ISequencer> InSequencer, UMovieScene* InMovieScene, FGuid InBindingGuid, ULevelSequenceEditorSubsystem* InLevelSequenceEditorSubsystem)
+	: SequencerPtr(InSequencer)
+	, MovieScene(InMovieScene)
+	, LevelSequenceEditorSubsystem(InLevelSequenceEditorSubsystem)
+	, BindingGuid(InBindingGuid)
+{
+
+}
+
+TSharedRef<IDetailCustomization> FMovieSceneBindingPropertyInfoListCustomization::MakeInstance(TWeakPtr<ISequencer> InSequencer, UMovieScene* InMovieScene, FGuid InBindingGuid, ULevelSequenceEditorSubsystem* InLevelSequenceEditorSubsystem)
+{
+	return MakeShareable(new FMovieSceneBindingPropertyInfoListCustomization(InSequencer, InMovieScene, InBindingGuid, InLevelSequenceEditorSubsystem));
+}
+
+void FMovieSceneBindingPropertyInfoListCustomization::CustomizeDetails(IDetailLayoutBuilder& DetailBuilder)
+{
+	IDetailCategoryBuilder& SectionCategory = DetailBuilder.EditCategory("Binding Properties", FText(), ECategoryPriority::Important);
+
+	TArray<TWeakObjectPtr<UObject>> Objects;
+	DetailBuilder.GetObjectsBeingCustomized(Objects);
+	TSharedPtr<ISequencer> Sequencer = SequencerPtr.Pin();
+
+	if (Objects.Num() > 0)
+	{
+		BindingList = Cast<UMovieSceneBindingPropertyInfoList>(Objects[0].Get());
+		if (BindingList)
+		{
+			if (Sequencer.IsValid())
+			{
+				UMovieSceneSequence* Sequence = MovieScene->GetTypedOuter<UMovieSceneSequence>();
+				if (BindingList->Bindings.Num() > 0)
+				{
+					// Grab the first one- we guarantee the binding types are the same.
+					FMovieSceneBindingPropertyInfo& BindingPropertyInfo = BindingList->Bindings[0];
+					bool bShowConvert = true;
+					if (FMovieScenePossessable* Possessable = MovieScene->FindPossessable(BindingGuid))
+					{
+						if (Possessable && Possessable->GetParent().IsValid())
+						{
+							bShowConvert = false;
+						}
+					}
+
+					int32 InitialIndex = 0;
+					BindingTypeNames.Add(MakeShared<FText>(LOCTEXT("BindingType_Possessable", "Possessable")));
+					TArrayView<const TSubclassOf<UMovieSceneCustomBinding>> SupportedBindingTypes = Sequencer->GetSupportedCustomBindingTypes();
+					int32 BindingTypeIndex = 1;
+					for (const TSubclassOf<UMovieSceneCustomBinding>& CustomBindingType : SupportedBindingTypes)
+					{
+						if (CustomBindingType)
+						{
+							if (BindingPropertyInfo.CustomBinding && BindingPropertyInfo.CustomBinding->GetClass() == CustomBindingType)
+							{
+								InitialIndex = BindingTypeIndex;
+							}
+							BindingTypeNames.Add(MakeShared<FText>(CustomBindingType->GetDefaultObject<UMovieSceneCustomBinding>()->GetBindingTypePrettyName()));
+						}
+						BindingTypeIndex++;
+					}
+
+					FDetailWidgetRow& BindingTypeRow = SectionCategory.AddCustomRow(FText::GetEmpty());
+					BindingTypeRow.NameContent()
+						[
+							SNew(STextBlock)
+							.Text(LOCTEXT("BindingPropertyType", "Binding Type"))
+						.ToolTipText(LOCTEXT("BindingPropertyType_Tooltip", "The type of binding for this object binding track entry"))
+						.Font(IDetailLayoutBuilder::GetDetailFont())
+						];
+					BindingTypeRow.ValueContent()
+						[
+							SNew(SComboBox<TSharedPtr<FText>>)
+							.OptionsSource(&BindingTypeNames)
+							.InitiallySelectedItem(BindingTypeNames[InitialIndex])
+							.IsEnabled(bShowConvert)
+							.OnSelectionChanged_Lambda([this, SupportedBindingTypes, &SectionCategory, &DetailBuilder](TSharedPtr<FText> NewValue, ESelectInfo::Type)
+							{
+								int32 CurrentIndex = BindingTypeNames.IndexOfByPredicate([&](const TSharedPtr<FText>& In) {
+									return In == NewValue;
+									});
+
+								for(FMovieSceneBindingPropertyInfo& BindingPropertyInfo : BindingList->Bindings)
+								{
+									if (CurrentIndex == 0)
+									{
+										// Possessable- clear everything
+										BindingPropertyInfo.ResolveFlags = ELocatorResolveFlags::None;
+										BindingPropertyInfo.Locator = FUniversalObjectLocator();
+										BindingPropertyInfo.CustomBinding = nullptr;
+									}
+									else
+									{
+										// Custom binding
+										UClass* BindingTypeClass = SupportedBindingTypes[CurrentIndex - 1];
+										BindingPropertyInfo.Locator = FUniversalObjectLocator();
+										BindingPropertyInfo.ResolveFlags = ELocatorResolveFlags::None;
+										BindingPropertyInfo.CustomBinding = NewObject<UMovieSceneCustomBinding>(MovieScene.Get(), BindingTypeClass, NAME_None, EObjectFlags::RF_Transient);
+									}
+								}
+
+								// Force reset the view
+								DetailBuilder.GetPropertyUtilities()->NotifyFinishedChangingProperties(FPropertyChangedEvent(FindFProperty<FProperty>(FMovieSceneBindingPropertyInfo::StaticStruct(), "Locator")));
+								SectionCategory.GetParentLayout().ForceRefreshDetails();
+							})
+						.OnGenerateWidget_Lambda([](TSharedPtr<FText> Item)
+							{
+								return SNew(STextBlock).Text(*Item.Get());
+							})
+								.InitiallySelectedItem(BindingTypeNames[InitialIndex])
+								[
+									SNew(STextBlock)
+									.Margin(FMargin(0.0f, 2.0f))
+								.Text_Lambda([this, InitialIndex]() { return *BindingTypeNames[InitialIndex].Get(); })
+								]
+						];
+
+					// Only show certain menus if we have a currently bound object
+					if (bShowConvert && MovieSceneHelpers::GetSingleBoundObject(Sequence, BindingGuid, Sequencer->GetSharedPlaybackState(), 0) != nullptr)
+					{
+						FDetailWidgetRow& ConvertToRow = SectionCategory.AddCustomRow(LOCTEXT("ConvertBindingTo", "Convert Binding To..."));
+						ConvertToRow.WholeRowContent()
+							[
+								SNew(SComboButton)
+								.OnGetMenuContent(FOnGetContent::CreateLambda([this, &DetailBuilder]() { return OnGetConvertMenuContent(DetailBuilder); }))
+							.ComboButtonStyle(FAppStyle::Get(), "SimpleComboButton")
+							.ButtonContent()
+							[
+								SNew(STextBlock).Text(LOCTEXT("ConvertBindingTo", "Convert Binding(s) To..."))
+							]
+							];
+					}
+				}
+			}
+		}
+	}
+}
+
+TSharedRef<SWidget> FMovieSceneBindingPropertyInfoListCustomization::OnGetConvertMenuContent(IDetailLayoutBuilder& DetailBuilder)
+{
+	FMenuBuilder MenuBuilder(true, nullptr, nullptr, true);
+
+	TSharedPtr<ISequencer> Sequencer = SequencerPtr.Pin();
+
+	if (Sequencer.IsValid() && BindingList)
+	{
+		TArray<FSequencerConvertBindingInfo> Bindings;
+		for (int32 BindingIndex = 0; BindingIndex < BindingList->Bindings.Num(); ++BindingIndex)
+		{
+			Bindings.Add({ BindingGuid, BindingIndex });
+		}
+
+		FSequencerUtilities::AddConvertBindingMenu(MenuBuilder, Sequencer.ToSharedRef(), Bindings, [this, &DetailBuilder]()
+			{
+				if (IDetailsView* DetailsView = DetailBuilder.GetDetailsView())
+				{
+					LevelSequenceEditorSubsystem->RefreshBindingDetails(DetailsView, BindingGuid);
+				}
+			});
+	}
+	return MenuBuilder.MakeWidget();
+}
+
 
 class FMovieSceneBindingPropertyInfoDetailCustomization : public IPropertyTypeCustomization
 {
@@ -88,7 +272,6 @@ public:
 
 private:
 
-	TSharedRef<SWidget> OnGetConvertMenuContent(IDetailChildrenBuilder& StructBuilder, IPropertyTypeCustomizationUtils& CustomizationUtils);
 	TSharedRef<SWidget> OnGetChangeClassMenuContent(IDetailChildrenBuilder& StructBuilder, IPropertyTypeCustomizationUtils& CustomizationUtils);
 
 	TWeakPtr<ISequencer> SequencerPtr;
@@ -98,8 +281,6 @@ private:
 	int32 BindingIndex = 0;
 
 	TSharedPtr<IPropertyHandle> StructPropertyHandle;
-
-	TArray<TSharedPtr<FText>> BindingTypeNames;
 };
 
 FMovieSceneBindingPropertyInfoDetailCustomization::FMovieSceneBindingPropertyInfoDetailCustomization(TWeakPtr<ISequencer> InSequencer, UMovieScene* InMovieScene, FGuid InBindingGuid, ULevelSequenceEditorSubsystem* InLevelSequenceEditorSubsystem)
@@ -138,117 +319,8 @@ void FMovieSceneBindingPropertyInfoDetailCustomization::CustomizeChildren(TShare
 		FMovieSceneBindingPropertyInfo* BindingPropertyInfo = (StructPtrs.Num() == 1) ? reinterpret_cast<FMovieSceneBindingPropertyInfo*>(StructPtrs[0]) : nullptr;
 		if (BindingPropertyInfo)
 		{
-			bool bShowConvert = true;
-			if (FMovieScenePossessable* Possessable = MovieScene->FindPossessable(BindingGuid))
-			{
-				if (Possessable && Possessable->GetParent().IsValid())
-				{
-					bShowConvert = false;
-				}
-			}
-
 			BindingIndex = InStructPropertyHandle->GetArrayIndex();
 
-			uint32 NumBindings = 0;
-			if (TSharedPtr<IPropertyHandle> ParentHandle = InStructPropertyHandle->GetParentHandle())
-			{
-				ParentHandle->GetNumChildren(NumBindings);
-			}
-
-			if (NumBindings > 1)
-			{
-				bShowConvert = false;
-			}
-
-			int32 InitialIndex = 0;
-			BindingTypeNames.Add(MakeShared<FText>(LOCTEXT("BindingType_Possessable", "Possessable")));
-			TArrayView<const TSubclassOf<UMovieSceneCustomBinding>> SupportedBindingTypes = Sequencer->GetSupportedCustomBindingTypes();
-			int32 BindingTypeIndex = 1;
-			for (const TSubclassOf<UMovieSceneCustomBinding>& CustomBindingType : SupportedBindingTypes)
-			{
-				if (CustomBindingType)
-				{
-					if (BindingPropertyInfo->CustomBinding && BindingPropertyInfo->CustomBinding->GetClass() == CustomBindingType)
-					{
-						InitialIndex = BindingTypeIndex;
-					}
-					BindingTypeNames.Add(MakeShared<FText>(CustomBindingType->GetDefaultObject<UMovieSceneCustomBinding>()->GetBindingTypePrettyName()));
-				}
-				BindingTypeIndex++;
-			}
-
-			FDetailWidgetRow& BindingTypeRow = StructBuilder.AddCustomRow(LOCTEXT("BindingPropertyType", "Binding Type"));
-			BindingTypeRow.NameContent()
-				[
-					SNew(STextBlock)
-					.Text(LOCTEXT("BindingPropertyType", "Binding Type"))
-					.ToolTipText(LOCTEXT("BindingPropertyType_Tooltip", "The type of binding for this object binding track entry"))
-					.Font(IDetailLayoutBuilder::GetDetailFont())
-				];
-			BindingTypeRow.ValueContent()
-				[
-					SNew(SComboBox<TSharedPtr<FText>>)
-					.OptionsSource(&BindingTypeNames)
-					.InitiallySelectedItem(BindingTypeNames[InitialIndex])
-					.IsEnabled(bShowConvert)
-					.OnSelectionChanged_Lambda([this, InStructPropertyHandle, SupportedBindingTypes, BindingPropertyInfo, &StructBuilder, &CustomizationUtils](TSharedPtr<FText> NewValue, ESelectInfo::Type)
-					{
-						int32 CurrentIndex = BindingTypeNames.IndexOfByPredicate([&](const TSharedPtr<FText>& In) {
-							return In == NewValue;
-							});
-						if (CurrentIndex == 0)
-						{
-							// Possessable- clear everything
-							BindingPropertyInfo->ResolveFlags = ELocatorResolveFlags::None;
-							BindingPropertyInfo->Locator = FUniversalObjectLocator();
-							BindingPropertyInfo->CustomBinding = nullptr;
-						}
-						else
-						{
-							// Custom binding
-							UClass* BindingTypeClass = SupportedBindingTypes[CurrentIndex - 1];
-							BindingPropertyInfo->Locator = FUniversalObjectLocator();
-							BindingPropertyInfo->ResolveFlags = ELocatorResolveFlags::None;
-							BindingPropertyInfo->CustomBinding = NewObject<UMovieSceneCustomBinding>(MovieScene.Get(), BindingTypeClass, NAME_None, EObjectFlags::RF_Transient);
-						}
-
-						// Force reset the view
-						TSharedPtr<IPropertyHandle> LocatorProperty = StructPropertyHandle->GetChildHandle(GET_MEMBER_NAME_CHECKED(FMovieSceneBindingPropertyInfo, Locator));
-						if (LocatorProperty.IsValid())
-						{
-							CustomizationUtils.GetPropertyUtilities()->NotifyFinishedChangingProperties(FPropertyChangedEvent(LocatorProperty->GetProperty()));
-						}
-						StructBuilder.GetParentCategory().GetParentLayout().ForceRefreshDetails();
-
-					})
-				.OnGenerateWidget_Lambda([](TSharedPtr<FText> Item)
-					{
-						return SNew(STextBlock).Text(*Item.Get());
-					})
-				.InitiallySelectedItem(BindingTypeNames[InitialIndex])
-				[
-					SNew(STextBlock)
-					.Margin(FMargin(0.0f, 2.0f))
-					.Text_Lambda([this, InitialIndex]() { return *BindingTypeNames[InitialIndex].Get(); })
-				]
-			];
-
-			// Only show certain menus if we have a currently bound object
-			if (bShowConvert && MovieSceneHelpers::GetSingleBoundObject(Sequence, BindingGuid, Sequencer->GetSharedPlaybackState(), BindingIndex) != nullptr)
-			{
-				FDetailWidgetRow& ConvertToRow = StructBuilder.AddCustomRow(LOCTEXT("ConvertBindingTo", "Convert Binding To..."));
-				ConvertToRow.WholeRowContent()
-					[
-						SNew(SComboButton)
-						.OnGetMenuContent(FOnGetContent::CreateLambda([this, &StructBuilder, &CustomizationUtils]() { return OnGetConvertMenuContent(StructBuilder, CustomizationUtils); }))
-					.ComboButtonStyle(FAppStyle::Get(), "SimpleComboButton")
-					.ButtonContent()
-					[
-						SNew(STextBlock).Text(LOCTEXT("ConvertBindingTo", "Convert Binding To..."))
-					]
-					];
-			}
-
 			// Show Change class and save default state menus for spawnables
 			if (MovieSceneHelpers::SupportsObjectTemplate(Sequence, BindingGuid, Sequencer->GetSharedPlaybackState(), BindingIndex))
 			{
@@ -283,9 +355,7 @@ void FMovieSceneBindingPropertyInfoDetailCustomization::CustomizeChildren(TShare
 					];
 			}
 
-
-
-			if (InitialIndex == 0)
+			if (!BindingPropertyInfo->CustomBinding)
 			{
 				// Show locator property
 				TSharedPtr<IPropertyHandle> LocatorProperty = StructPropertyHandle->GetChildHandle(GET_MEMBER_NAME_CHECKED(FMovieSceneBindingPropertyInfo, Locator));
@@ -313,25 +383,6 @@ void FMovieSceneBindingPropertyInfoDetailCustomization::CustomizeChildren(TShare
 	}
 }
 
-TSharedRef<SWidget> FMovieSceneBindingPropertyInfoDetailCustomization::OnGetConvertMenuContent(IDetailChildrenBuilder& StructBuilder, IPropertyTypeCustomizationUtils& CustomizationUtils)
-{
-	FMenuBuilder MenuBuilder(true, nullptr, nullptr, true);
-
-	TSharedPtr<ISequencer> Sequencer = SequencerPtr.Pin();
-	if (Sequencer.IsValid() && StructPropertyHandle->IsValidHandle())
-	{
-		TArray<FSequencerConvertBindingInfo> Bindings;
-		Bindings.Add({BindingGuid, BindingIndex});
-		FSequencerUtilities::AddConvertBindingMenu(MenuBuilder, Sequencer.ToSharedRef(), Bindings, [this, &StructBuilder, &CustomizationUtils]()
-			{
-				if (IDetailsView* DetailsView = StructBuilder.GetParentCategory().GetParentLayout().GetDetailsView())
-				{
-					LevelSequenceEditorSubsystem->RefreshBindingDetails(DetailsView, BindingGuid);
-				}				
-			});
-	}
-	return MenuBuilder.MakeWidget();
-}
 
 TSharedRef<SWidget> FMovieSceneBindingPropertyInfoDetailCustomization::OnGetChangeClassMenuContent(IDetailChildrenBuilder& StructBuilder, IPropertyTypeCustomizationUtils& CustomizationUtils)
 {
@@ -543,6 +594,9 @@ void ULevelSequenceEditorSubsystem::AddBindingDetailCustomizations(TSharedRef<ID
 		if (MovieScene)
 		{
 			FPropertyEditorModule& PropertyEditor = FModuleManager::Get().LoadModuleChecked<FPropertyEditorModule>(TEXT("PropertyEditor"));
+
+			DetailsView->RegisterInstancedCustomPropertyLayout(UMovieSceneBindingPropertyInfoList::StaticClass(), FOnGetDetailCustomizationInstance::CreateStatic(&FMovieSceneBindingPropertyInfoListCustomization::MakeInstance, ActiveSequencer.ToWeakPtr(), MovieScene, BindingGuid, this));
+			
 			DetailsView->RegisterInstancedCustomPropertyTypeLayout(FMovieSceneBindingPropertyInfo::StaticStruct()->GetFName(), FOnGetPropertyTypeCustomizationInstance::CreateLambda([](TWeakPtr<ISequencer> InSequencer, UMovieScene* InMovieScene, FGuid InBindingGuid, ULevelSequenceEditorSubsystem* LevelSequenceEditorSubsystem)
 				{
 					return MakeShared<FMovieSceneBindingPropertyInfoDetailCustomization>(InSequencer, InMovieScene, InBindingGuid, LevelSequenceEditorSubsystem);
diff --git a/Engine/Plugins/MovieScene/LevelSequenceEditor/Source/LevelSequenceEditor/Public/LevelSequenceEditorSubsystem.h b/Engine/Plugins/MovieScene/LevelSequenceEditor/Source/LevelSequenceEditor/Public/LevelSequenceEditorSubsystem.h
--- a/Engine/Plugins/MovieScene/LevelSequenceEditor/Source/LevelSequenceEditor/Public/LevelSequenceEditorSubsystem.h
+++ b/Engine/Plugins/MovieScene/LevelSequenceEditor/Source/LevelSequenceEditor/Public/LevelSequenceEditorSubsystem.h
@@ -74,7 +74,7 @@ class UMovieSceneBindingPropertyInfoList : public UObject
 	GENERATED_BODY()
 public:
 	// List of locator info for a particular binding
-	UPROPERTY(EditAnywhere, Category = "Default")
+	UPROPERTY(EditAnywhere, Category = "Binding Properties")
 	TArray<FMovieSceneBindingPropertyInfo> Bindings;
 };
 
