Sequencer: Fix conflicting LOCTEXT tag.

diff --git a/Engine/Plugins/MovieScene/LevelSequenceEditor/Source/LevelSequenceEditor/Private/LevelSequenceEditorSubsystem.cpp b/Engine/Plugins/MovieScene/LevelSequenceEditor/Source/LevelSequenceEditor/Private/LevelSequenceEditorSubsystem.cpp
--- a/Engine/Plugins/MovieScene/LevelSequenceEditor/Source/LevelSequenceEditor/Private/LevelSequenceEditorSubsystem.cpp
+++ b/Engine/Plugins/MovieScene/LevelSequenceEditor/Source/LevelSequenceEditor/Private/LevelSequenceEditorSubsystem.cpp
@@ -216,7 +216,7 @@ void FMovieSceneBindingPropertyInfoListCustomization::CustomizeDetails(IDetailLa
 					// Only show certain menus if we have a currently bound object
 					if (bShowConvert && MovieSceneHelpers::GetSingleBoundObject(Sequence, BindingGuid, Sequencer->GetSharedPlaybackState(), 0) != nullptr)
 					{
-						FDetailWidgetRow& ConvertToRow = SectionCategory.AddCustomRow(LOCTEXT("ConvertBindingTo", "Convert Binding To..."));
+						FDetailWidgetRow& ConvertToRow = SectionCategory.AddCustomRow(LOCTEXT("ConvertBindingTo", "Convert Binding(s) To..."));
 						ConvertToRow.WholeRowContent()
 							[
 								SNew(SComboButton)
