Sequencer: Fixed 'Keep Playhead in Playback Range While Scrubbing' not working correctly for cropped sub sequences

diff --git a/Engine/Plugins/MovieScene/LevelSequenceEditor/Source/LevelSequenceEditor/Private/CinematicViewport/SCinematicTransportRange.cpp b/Engine/Plugins/MovieScene/LevelSequenceEditor/Source/LevelSequenceEditor/Private/CinematicViewport/SCinematicTransportRange.cpp
--- a/Engine/Plugins/MovieScene/LevelSequenceEditor/Source/LevelSequenceEditor/Private/CinematicViewport/SCinematicTransportRange.cpp
+++ b/Engine/Plugins/MovieScene/LevelSequenceEditor/Source/LevelSequenceEditor/Private/CinematicViewport/SCinematicTransportRange.cpp
@@ -89,7 +89,7 @@ void SCinematicTransportRange::SetTime(const FGeometry& MyGeometry, const FPoint
 		// Clamp first, snap to frame last
 		if (Sequencer->GetSequencerSettings()->ShouldKeepCursorInPlayRangeWhileScrubbing())
 		{
-			TRange<FFrameNumber> PlaybackRange = Sequencer->GetFocusedMovieSceneSequence()->GetMovieScene()->GetPlaybackRange();
+			TRange<FFrameNumber> PlaybackRange = Sequencer->GetSubSequenceRange().Get(Sequencer->GetRootMovieSceneSequence()->GetMovieScene()->GetPlaybackRange());
 			ScrubTime = UE::MovieScene::ClampToDiscreteRange(ScrubTime, PlaybackRange);
 		}
 
diff --git a/Engine/Source/Editor/Sequencer/Private/Sequencer.cpp b/Engine/Source/Editor/Sequencer/Private/Sequencer.cpp
--- a/Engine/Source/Editor/Sequencer/Private/Sequencer.cpp
+++ b/Engine/Source/Editor/Sequencer/Private/Sequencer.cpp
@@ -3350,7 +3350,7 @@ void FSequencer::SnapSequencerTime(FFrameTime& ScrubTime)
 	// Clamp first, snap to frame last
 	if (GetSequencerSettings()->ShouldKeepCursorInPlayRangeWhileScrubbing())
 	{
-		TRange<FFrameNumber> PlaybackRange =GetFocusedMovieSceneSequence()->GetMovieScene()->GetPlaybackRange();
+		TRange<FFrameNumber> PlaybackRange = GetSubSequenceRange().Get(GetRootMovieSceneSequence()->GetMovieScene()->GetPlaybackRange());
 		ScrubTime = UE::MovieScene::ClampToDiscreteRange(ScrubTime, PlaybackRange);
 	}
 
