Sequencer- fix crash when converting to spawnable director blueprint binding.

diff --git a/Engine/Plugins/MovieScene/LevelSequenceEditor/Source/LevelSequenceEditor/Private/Misc/LevelSequenceEditorActorSpawner.cpp b/Engine/Plugins/MovieScene/LevelSequenceEditor/Source/LevelSequenceEditor/Private/Misc/LevelSequenceEditorActorSpawner.cpp
--- a/Engine/Plugins/MovieScene/LevelSequenceEditor/Source/LevelSequenceEditor/Private/Misc/LevelSequenceEditorActorSpawner.cpp
+++ b/Engine/Plugins/MovieScene/LevelSequenceEditor/Source/LevelSequenceEditor/Private/Misc/LevelSequenceEditorActorSpawner.cpp
@@ -202,7 +202,6 @@ bool FLevelSequenceEditorActorSpawner::CanSetupDefaultsForSpawnable(UObject* Spa
 void FLevelSequenceEditorActorSpawner::SetupDefaultsForSpawnable(UObject* SpawnedObject, const FGuid& Guid, const TOptional<FTransformData>& TransformData, TSharedRef<ISequencer> Sequencer, USequencerSettings* Settings)
 {
 	TOptional<FTransformData> DefaultTransform = TransformData;
-	DefaultTransform->Scale = FVector(1.0f, 1.0f, 1.0f);
 
 	AActor* SpawnedActor = Cast<AActor>(SpawnedObject);
 	if (SpawnedActor)
@@ -228,6 +227,10 @@ void FLevelSequenceEditorActorSpawner::SetupDefaultsForSpawnable(UObject* Spawne
 		GEditor->SelectNone(bNotifySelectionChanged, bDeselectBSP, bWarnAboutTooManyActors);
 		GEditor->SelectActor(SpawnedActor, true, bNotifySelectionChanged, bSelectEvenIfHidden);
 	}
+	else if (DefaultTransform.IsSet())
+	{
+		DefaultTransform->Scale = FVector(1.0f, 1.0f, 1.0f);
+	}
 
 	UMovieSceneSequence* Sequence = Sequencer->GetFocusedMovieSceneSequence();
 	UMovieScene* OwnerMovieScene = Sequence->GetMovieScene();
