Sequencer- fix failing autotest and issue with UI code holding a reference to a sequence after Sequencer has closed. This started happening when we started showing Binding Property menu outside of a right-click menu.

diff --git a/Engine/Plugins/MovieScene/LevelSequenceEditor/Source/LevelSequenceEditor/Private/LevelSequenceEditorSubsystem.cpp b/Engine/Plugins/MovieScene/LevelSequenceEditor/Source/LevelSequenceEditor/Private/LevelSequenceEditorSubsystem.cpp
--- a/Engine/Plugins/MovieScene/LevelSequenceEditor/Source/LevelSequenceEditor/Private/LevelSequenceEditorSubsystem.cpp
+++ b/Engine/Plugins/MovieScene/LevelSequenceEditor/Source/LevelSequenceEditor/Private/LevelSequenceEditorSubsystem.cpp
@@ -515,12 +515,12 @@ void ULevelSequenceEditorSubsystem::Initialize(FSubsystemCollectionBase& Collect
 	// For now we have the binding properties being a separate menu. When the UX is worked out we will likely merge the AssignActor menu away.
 	BindingPropertiesMenuExtender = MakeShareable(new FExtender);
 	BindingPropertiesMenuExtender->AddMenuExtension("Possessable", EExtensionHook::First, CommandList, FMenuExtensionDelegate::CreateLambda([this](FMenuBuilder& MenuBuilder) {
-		
+
 		FFormatNamedArguments Args;
 		MenuBuilder.AddSubMenu(
 			FText::Format(LOCTEXT("BindingProperties", "Binding Properties"), Args),
 			FText::Format(LOCTEXT("BindingPropertiesTooltip", "Modify the actor and object bindings for this track"), Args),
-			FNewMenuDelegate::CreateLambda([this](FMenuBuilder& SubMenuBuilder) { AddBindingPropertiesMenu(SubMenuBuilder); } ));
+			FNewMenuDelegate::CreateLambda([this](FMenuBuilder& SubMenuBuilder) { AddBindingPropertiesMenu(SubMenuBuilder); }));
 		}));
 
 	SequencerModule.GetObjectBindingContextMenuExtensibilityManager()->AddExtender(BindingPropertiesMenuExtender);
@@ -532,7 +532,7 @@ void ULevelSequenceEditorSubsystem::Initialize(FSubsystemCollectionBase& Collect
 		{
 			return;
 		}
-		
+
 		TArray<FName> ComponentNames;
 		GetRebindComponentNames(ComponentNames);
 		if (ComponentNames.Num() > 0)
@@ -541,13 +541,13 @@ void ULevelSequenceEditorSubsystem::Initialize(FSubsystemCollectionBase& Collect
 			MenuBuilder.AddSubMenu(
 				FText::Format(LOCTEXT("RebindComponent", "Rebind Component"), Args),
 				FText::Format(LOCTEXT("RebindComponentTooltip", "Rebind component by moving the tracks from one component to another component."), Args),
-				FNewMenuDelegate::CreateLambda([this](FMenuBuilder& SubMenuBuilder) { RebindComponentMenu(SubMenuBuilder); } ));
+				FNewMenuDelegate::CreateLambda([this](FMenuBuilder& SubMenuBuilder) { RebindComponentMenu(SubMenuBuilder); }));
 		}
-	}));
+		}));
 
 	BindingPropertiesMenuExtender->AddMenuExtension("CustomBinding", EExtensionHook::First, CommandList, FMenuExtensionDelegate::CreateLambda([this](FMenuBuilder& MenuBuilder) {
 		// Only add menu entries where the focused sequence is a ULevelSequence
-		
+
 		// Add instanced detail customizations
 
 		FFormatNamedArguments Args;
@@ -563,26 +563,26 @@ void ULevelSequenceEditorSubsystem::Initialize(FSubsystemCollectionBase& Collect
 
 	SidebarMenuExtender->AddMenuExtension(TEXT("Possessable"), EExtensionHook::First, CommandList,
 		FMenuExtensionDelegate::CreateLambda([this](FMenuBuilder& MenuBuilder)
-		{
-			// Only add menu entries where the focused sequence is a ULevelSequence
-			if (GetActiveSequencer())
 			{
-				AddBindingPropertiesSidebar(MenuBuilder);
-			}
+				// Only add menu entries where the focused sequence is a ULevelSequence
+				if (GetActiveSequencer())
+				{
+					AddBindingPropertiesSidebar(MenuBuilder);
+				}
 
-			TArray<FName> ComponentNames;
-			GetRebindComponentNames(ComponentNames);
-			if (ComponentNames.Num() > 0)
-			{
-				RebindComponentMenu(MenuBuilder);
-			}
-		}));
+				TArray<FName> ComponentNames;
+				GetRebindComponentNames(ComponentNames);
+				if (ComponentNames.Num() > 0)
+				{
+					RebindComponentMenu(MenuBuilder);
+				}
+			}));
 
 	SidebarMenuExtender->AddMenuExtension(TEXT("CustomBinding"), EExtensionHook::First, CommandList,
 		FMenuExtensionDelegate::CreateLambda([this](FMenuBuilder& MenuBuilder)
-		{
-			AddBindingPropertiesMenu(MenuBuilder);
-		}));
+			{
+				AddBindingPropertiesMenu(MenuBuilder);
+			}));
 
 	SequencerModule.GetSidebarExtensibilityManager()->AddExtender(SidebarMenuExtender);
 }
@@ -597,6 +597,7 @@ void ULevelSequenceEditorSubsystem::Deinitialize()
 		SequencerModulePtr->UnregisterOnSequencerCreated(OnSequencerCreatedHandle);
 	}
 
+	BindingPropertyInfoList = nullptr;
 	if (FSlateApplication::IsInitialized())
 	{
 		FSlateApplication::Get().OnMenuBeingDestroyed().RemoveAll(this);
@@ -609,6 +610,12 @@ void ULevelSequenceEditorSubsystem::OnSequencerCreated(TSharedRef<ISequencer> In
 	UE_LOG(LogLevelSequenceEditor, VeryVerbose, TEXT("ULevelSequenceEditorSubsystem::OnSequencerCreated"));
 
 	Sequencers.Add(TWeakPtr<ISequencer>(InSequencer));
+	InSequencer->OnCloseEvent().AddUObject(this, &ULevelSequenceEditorSubsystem::OnSequencerClosed);
+}
+
+void ULevelSequenceEditorSubsystem::OnSequencerClosed(TSharedRef<ISequencer> InSequencer)
+{
+	BindingPropertyInfoList = nullptr;
 }
 
 void ULevelSequenceEditorSubsystem::AddBindingDetailCustomizations(TSharedRef<IDetailsView> DetailsView, TSharedPtr<ISequencer> ActiveSequencer, FGuid BindingGuid)
diff --git a/Engine/Plugins/MovieScene/LevelSequenceEditor/Source/LevelSequenceEditor/Public/LevelSequenceEditorSubsystem.h b/Engine/Plugins/MovieScene/LevelSequenceEditor/Source/LevelSequenceEditor/Public/LevelSequenceEditorSubsystem.h
--- a/Engine/Plugins/MovieScene/LevelSequenceEditor/Source/LevelSequenceEditor/Public/LevelSequenceEditorSubsystem.h
+++ b/Engine/Plugins/MovieScene/LevelSequenceEditor/Source/LevelSequenceEditor/Public/LevelSequenceEditorSubsystem.h
@@ -94,6 +94,8 @@ public:
 
 	void OnSequencerCreated(TSharedRef<ISequencer> InSequencer);
 
+	void OnSequencerClosed(TSharedRef<ISequencer> InSequencer);
+
 	/** Retrieve the scripting layer */
 	UFUNCTION(BlueprintPure, Category = "Level Sequence Editor")
 	USequencerModuleScriptingLayer* GetScriptingLayer();
