Sequencer: Show level instance actors on actor pickers if editor-paths are enabled

diff --git a/Engine/Plugins/MovieScene/LevelSequenceEditor/Source/LevelSequenceEditor/Private/Misc/LevelSequenceEditorActorBinding.cpp b/Engine/Plugins/MovieScene/LevelSequenceEditor/Source/LevelSequenceEditor/Private/Misc/LevelSequenceEditorActorBinding.cpp
--- a/Engine/Plugins/MovieScene/LevelSequenceEditor/Source/LevelSequenceEditor/Private/Misc/LevelSequenceEditorActorBinding.cpp
+++ b/Engine/Plugins/MovieScene/LevelSequenceEditor/Source/LevelSequenceEditor/Private/Misc/LevelSequenceEditorActorBinding.cpp
@@ -12,6 +12,7 @@
 #include "Framework/MultiBox/MultiBoxBuilder.h"
 #include "Framework/Application/SlateApplication.h"
 #include "Modules/ModuleManager.h"
+#include "Misc/EditorPathHelper.h"
 #include "Editor.h"
 #include "Selection.h"
 
@@ -130,6 +131,8 @@ void FLevelSequenceEditorActorBinding::AddPossessActorMenuExtensions(FMenuBuilde
 		InitOptions.Filters->AddFilterPredicate<FActorTreeItem>(FActorTreeItem::FFilterPredicate::CreateLambda(IsActorValidForPossession, ExistingPossessedObjects));
 	}
 
+	const bool bHideLevelInstanceHierarchy = !FEditorPathHelper::IsEnabled();
+
 	// actor selector to allow the user to choose an actor
 	FSceneOutlinerModule& SceneOutlinerModule = FModuleManager::LoadModuleChecked<FSceneOutlinerModule>("SceneOutliner");
 	TSharedRef< SWidget > MiniSceneOutliner =
@@ -143,7 +146,9 @@ void FLevelSequenceEditorActorBinding::AddPossessActorMenuExtensions(FMenuBuilde
 					// Create a new binding for this actor
 					FSlateApplication::Get().DismissAllMenus();
 					AddActorsToSequencer(&Actor, 1);
-				})
+				}),
+				nullptr,
+				bHideLevelInstanceHierarchy
 			)
 		];
 
