Sequencer: Only add Selected from Content Browser for level sequence editors and UMG

#rnx

diff --git a/Engine/Plugins/MovieScene/LevelSequenceEditor/Source/LevelSequenceEditor/Private/LevelSequenceEditorToolkit.cpp b/Engine/Plugins/MovieScene/LevelSequenceEditor/Source/LevelSequenceEditor/Private/LevelSequenceEditorToolkit.cpp
--- a/Engine/Plugins/MovieScene/LevelSequenceEditor/Source/LevelSequenceEditor/Private/LevelSequenceEditorToolkit.cpp
+++ b/Engine/Plugins/MovieScene/LevelSequenceEditor/Source/LevelSequenceEditor/Private/LevelSequenceEditorToolkit.cpp
@@ -175,6 +175,7 @@ void FLevelSequenceEditorToolkit::Initialize(const EToolkitMode::Type Mode, cons
 		SequencerInitParams.HostCapabilities.bSupportsSaveMovieSceneAsset = true;
 		SequencerInitParams.HostCapabilities.bSupportsRecording = true;
 		SequencerInitParams.HostCapabilities.bSupportsRenderMovie = true;
+		SequencerInitParams.HostCapabilities.bSupportsAddFromContentBrowser = true;
 	}
 
 	ExtendSequencerToolbar("Sequencer.MainToolBar");
@@ -588,9 +589,10 @@ void FLevelSequenceEditorToolkit::AddDefaultTracksForActor(AActor& Actor, const
 				break;
 			}
 
+			FGuid ComponentBinding = Sequencer->GetHandleToObject(PropertyOwner);
+
 			if (bReplaceWithTransformTrack)
 			{
-				FGuid ComponentBinding = Sequencer->GetHandleToObject(PropertyOwner);
 				UClass* TrackClass = UMovieScene3DTransformTrack::StaticClass();
 				UMovieSceneTrack* NewTrack = MovieScene->FindTrack(TrackClass, ComponentBinding);
 				if (!NewTrack)
@@ -601,6 +603,29 @@ void FLevelSequenceEditorToolkit::AddDefaultTracksForActor(AActor& Actor, const
 				continue;
 			}
 
+			bool bFoundPropertyTrack = false;
+			if (ComponentBinding.IsValid())
+			{
+				TArray<UMovieSceneTrack*> Tracks = MovieScene->FindTracks(UMovieScenePropertyTrack::StaticClass(), ComponentBinding, NAME_None);
+				for (UMovieSceneTrack* Track : Tracks)
+				{
+					UMovieScenePropertyTrack* PropertyTrack = Cast<UMovieScenePropertyTrack>(Track);
+					if (PropertyTrack)
+					{
+						if (PropertyTrack->GetPropertyPath() == PropertyTrackSettings.PropertyPath)
+						{
+							bFoundPropertyTrack = true;
+							break;
+						}
+					}
+				}
+			}
+
+			if (bFoundPropertyTrack)
+			{
+				continue;
+			}
+
 			if (!Sequencer->CanKeyProperty(FCanKeyPropertyParams(PropertyOwner->GetClass(), *PropertyPath)))
 			{
 				continue;
diff --git a/Engine/Source/Editor/Sequencer/Private/SSequencer.cpp b/Engine/Source/Editor/Sequencer/Private/SSequencer.cpp
--- a/Engine/Source/Editor/Sequencer/Private/SSequencer.cpp
+++ b/Engine/Source/Editor/Sequencer/Private/SSequencer.cpp
@@ -1687,13 +1687,16 @@ TSharedRef<SWidget> SSequencer::MakeAddMenu()
 	TSharedPtr<FExtender> Extender = FExtender::Combine(AddMenuExtenders);
 	FMenuBuilder MenuBuilder(true, nullptr, Extender);
 
-	MenuBuilder.AddMenuEntry(
-		LOCTEXT("SelectedFromContentBrowser", "Selected from Content Browser"),
-		LOCTEXT("SelectedFromContentBrowserToolTip", "Add selected content from the content browser"),
-		FSlateIcon(FAppStyle::GetAppStyleSetName(), "Icons.Use"),
-		FUIAction(
-			FExecuteAction::CreateRaw(this, &SSequencer::AddFromContentBrowser),
-			FCanExecuteAction::CreateRaw(this, &SSequencer::CanAddFromContentBrowser)));
+	if (SequencerPtr.Pin()->GetHostCapabilities().bSupportsAddFromContentBrowser)
+	{
+		MenuBuilder.AddMenuEntry(
+			LOCTEXT("SelectedFromContentBrowser", "Selected from Content Browser"),
+			LOCTEXT("SelectedFromContentBrowserToolTip", "Add selected content from the content browser"),
+			FSlateIcon(FAppStyle::GetAppStyleSetName(), "Icons.Use"),
+			FUIAction(
+				FExecuteAction::CreateRaw(this, &SSequencer::AddFromContentBrowser),
+				FCanExecuteAction::CreateRaw(this, &SSequencer::CanAddFromContentBrowser)));
+	}
 
 	{
 		TSharedPtr<FSequencer> Sequencer = SequencerPtr.Pin();
diff --git a/Engine/Source/Editor/Sequencer/Public/ISequencerModule.h b/Engine/Source/Editor/Sequencer/Public/ISequencerModule.h
--- a/Engine/Source/Editor/Sequencer/Public/ISequencerModule.h
+++ b/Engine/Source/Editor/Sequencer/Public/ISequencerModule.h
@@ -147,11 +147,15 @@ struct FSequencerHostCapabilities
 	/** Do we support rendering to a movie */
 	bool bSupportsRenderMovie;
 
+	/** Do we support adding content from the content browser */
+	bool bSupportsAddFromContentBrowser;
+
 	FSequencerHostCapabilities()
 		: bSupportsSaveMovieSceneAsset(false)
 		, bSupportsCurveEditor(false)
 		, bSupportsRecording(false)
 		, bSupportsRenderMovie(false)
+		, bSupportsAddFromContentBrowser(false)
 	{}
 };
 
diff --git a/Engine/Source/Editor/UMGEditor/Private/WidgetBlueprintEditor.cpp b/Engine/Source/Editor/UMGEditor/Private/WidgetBlueprintEditor.cpp
--- a/Engine/Source/Editor/UMGEditor/Private/WidgetBlueprintEditor.cpp
+++ b/Engine/Source/Editor/UMGEditor/Private/WidgetBlueprintEditor.cpp
@@ -1454,6 +1454,7 @@ TSharedPtr<ISequencer> FWidgetBlueprintEditor::CreateSequencerWidgetInternal()
 		SequencerInitParams.EventContexts = TAttribute<TArray<UObject*>>(this, &FWidgetBlueprintEditor::GetAnimationEventContexts);
 
 		SequencerInitParams.HostCapabilities.bSupportsCurveEditor = true;
+		SequencerInitParams.HostCapabilities.bSupportsAddFromContentBrowser = true;
 	};
 
 	TSharedPtr<ISequencer> Sequencer = FModuleManager::LoadModuleChecked<ISequencerModule>("Sequencer").CreateSequencer(SequencerInitParams);
