Sequencer: Respect the selected level sequence actor as the PlaybackClient when using the 'Open Level Sequence' button on Level Sequence Actor details. This ensures that the same level sequence placed multiple times will respect the context of that actor when viewing the sequence via opening Sequencer from that actor. This means that transform origin and binding overrides, which can be per Level Sequence Actor, work correctly in this context.

diff --git a/Engine/Plugins/MovieScene/LevelSequenceEditor/Source/LevelSequenceEditor/Private/Misc/LevelSequencePlaybackContext.cpp b/Engine/Plugins/MovieScene/LevelSequenceEditor/Source/LevelSequenceEditor/Private/Misc/LevelSequencePlaybackContext.cpp
--- a/Engine/Plugins/MovieScene/LevelSequenceEditor/Source/LevelSequenceEditor/Private/Misc/LevelSequencePlaybackContext.cpp
+++ b/Engine/Plugins/MovieScene/LevelSequenceEditor/Source/LevelSequenceEditor/Private/Misc/LevelSequencePlaybackContext.cpp
@@ -19,6 +19,7 @@
 #include "LevelInstance/LevelInstanceInterface.h"
 #include "LevelInstance/LevelInstanceSubsystem.h"
 #include "Editor.h"
+#include "Selection.h"
 
 #define LOCTEXT_NAMESPACE "LevelSequencePlaybackContext"
 
@@ -305,6 +306,16 @@ FLevelSequencePlaybackContext::FContextAndClient FLevelSequencePlaybackContext::
 
 	if (ensure(EditorWorld))
 	{
+		// First search selected actors. If one is selected, use that as the client.
+		for (FSelectionIterator It(GEditor->GetSelectedActorIterator()); It; ++It)
+		{
+			if (ALevelSequenceActor* LSA = Cast<ALevelSequenceActor>(*It))
+			{
+				return FContextAndClient(EditorWorld, LSA);
+			}
+		}
+
+		// Otherwise, attempt to find one in the world.
 		TArray<ALevelSequenceActor*> LevelSequenceActors;
 		UE::MovieScene::FindLevelSequenceActors(EditorWorld, InLevelSequence, LevelSequenceActors);
 		return FContextAndClient(EditorWorld, (LevelSequenceActors.Num() > 0 ? LevelSequenceActors[0] : nullptr));
