Sequencer: Custom binding UX polish.
There is now a new menu for changing or converting custom bindings. This includes the following additions/changes:
* Icons (spawnable/replaceable, or whatever custom icon the user might have chosen) alongside the entries
* Director Blueprint replaceable/spawnable has submenu for choosiing the endpoint right away
* Separate section for blueprint classes mimicking what we do for conditions- allowing creation or choice of a user-created binding type

This required a minor refactor of the core director blueprint customization allowing it to be used as a helper when a property editor doesn't (necessarily) exist.

Also changed the custom binding user-defined priorities to be a little less 'magic number'.

Finally, I created a new custom replaceable binding blueprintable base class allowing users to create their own replaceable binding type from a blueprint. We could do the same for spawnable types later if desired, though I think that should be less common.

diff --git a/Engine/Plugins/Experimental/ChaosCaching/Source/ChaosCaching/Public/Chaos/Sequencer/MovieSceneSpawnableChaosCacheBinding.h b/Engine/Plugins/Experimental/ChaosCaching/Source/ChaosCaching/Public/Chaos/Sequencer/MovieSceneSpawnableChaosCacheBinding.h
--- a/Engine/Plugins/Experimental/ChaosCaching/Source/ChaosCaching/Public/Chaos/Sequencer/MovieSceneSpawnableChaosCacheBinding.h
+++ b/Engine/Plugins/Experimental/ChaosCaching/Source/ChaosCaching/Public/Chaos/Sequencer/MovieSceneSpawnableChaosCacheBinding.h
@@ -32,7 +32,7 @@ public:
 protected:
 
 	// Ensures we override priority of UMovieSceneSpawnableActorBinding
-	int32 GetCustomBindingPriority() const override { return 11; }
+	int32 GetCustomBindingPriority() const override { return BaseCustomPriority; }
 
 };
 
diff --git a/Engine/Plugins/MovieScene/LevelSequenceEditor/Source/LevelSequenceEditor/Private/LevelSequenceEditorSubsystem.cpp b/Engine/Plugins/MovieScene/LevelSequenceEditor/Source/LevelSequenceEditor/Private/LevelSequenceEditorSubsystem.cpp
--- a/Engine/Plugins/MovieScene/LevelSequenceEditor/Source/LevelSequenceEditor/Private/LevelSequenceEditorSubsystem.cpp
+++ b/Engine/Plugins/MovieScene/LevelSequenceEditor/Source/LevelSequenceEditor/Private/LevelSequenceEditorSubsystem.cpp
@@ -70,6 +70,19 @@
 #include "MovieSceneDynamicBindingCustomization.h"
 #include "Conditions/MovieSceneConditionCustomization.h"
 #include "Conditions/MovieSceneDirectorBlueprintConditionCustomization.h"
+#include "MVVM/Selection/Selection.h"
+#include "UnrealEdGlobals.h"
+#include "Misc/FeedbackContext.h"
+#include "Variants/MovieSceneTimeWarpVariant.h"
+#include "Variants/MovieSceneTimeWarpGetter.h"
+#include "UObject/UObjectIterator.h"
+#include "ObjectTools.h"
+#include "Bindings/MovieSceneSpawnableDirectorBlueprintBinding.h"
+#include "Bindings/MovieSceneReplaceableDirectorBlueprintBinding.h"
+#include "Bindings/MovieSceneReplaceableActorBinding.h"
+#include "AssetToolsModule.h"
+#include "Kismet2/KismetEditorUtilities.h"
+#include "Kismet2/BlueprintEditorUtils.h"
 
 #include UE_INLINE_GENERATED_CPP_BY_NAME(LevelSequenceEditorSubsystem)
 
@@ -91,6 +104,10 @@ private:
 
 	TSharedRef<SWidget> OnGetConvertMenuContent(IDetailLayoutBuilder& DetailBuilder);
 
+	TSharedRef<SWidget> GenerateBindingTypePicker(IDetailLayoutBuilder& DetailBuilder);
+	const FSlateBrush* GetBindingTypeIcon() const;
+	FText GetBindingTypeValueAsString() const;
+
 	TWeakPtr<ISequencer> SequencerPtr;
 	TObjectPtr<UMovieScene> MovieScene = nullptr;
 	TObjectPtr<ULevelSequenceEditorSubsystem> LevelSequenceEditorSubsystem;
@@ -133,7 +150,6 @@ void FMovieSceneBindingPropertyInfoListCustomization::CustomizeDetails(IDetailLa
 				if (BindingList->Bindings.Num() > 0)
 				{
 					// Grab the first one- we guarantee the binding types are the same.
-					FMovieSceneBindingPropertyInfo& BindingPropertyInfo = BindingList->Bindings[0];
 					bool bShowConvert = true;
 					if (FMovieScenePossessable* Possessable = MovieScene->FindPossessable(BindingGuid))
 					{
@@ -143,23 +159,6 @@ void FMovieSceneBindingPropertyInfoListCustomization::CustomizeDetails(IDetailLa
 						}
 					}
 
-					int32 InitialIndex = 0;
-					BindingTypeNames.Add(MakeShared<FText>(LOCTEXT("BindingType_Possessable", "Possessable")));
-					TArrayView<const TSubclassOf<UMovieSceneCustomBinding>> SupportedBindingTypes = Sequencer->GetSupportedCustomBindingTypes();
-					int32 BindingTypeIndex = 1;
-					for (const TSubclassOf<UMovieSceneCustomBinding>& CustomBindingType : SupportedBindingTypes)
-					{
-						if (CustomBindingType)
-						{
-							if (BindingPropertyInfo.CustomBinding && BindingPropertyInfo.CustomBinding->GetClass() == CustomBindingType)
-							{
-								InitialIndex = BindingTypeIndex;
-							}
-							BindingTypeNames.Add(MakeShared<FText>(CustomBindingType->GetDefaultObject<UMovieSceneCustomBinding>()->GetBindingTypePrettyName()));
-						}
-						BindingTypeIndex++;
-					}
-
 					FDetailWidgetRow& BindingTypeRow = SectionCategory.AddCustomRow(FText::GetEmpty());
 					BindingTypeRow.NameContent()
 						[
@@ -170,49 +169,27 @@ void FMovieSceneBindingPropertyInfoListCustomization::CustomizeDetails(IDetailLa
 						];
 					BindingTypeRow.ValueContent()
 						[
-							SNew(SComboBox<TSharedPtr<FText>>)
-							.OptionsSource(&BindingTypeNames)
-							.InitiallySelectedItem(BindingTypeNames[InitialIndex])
-							.IsEnabled(bShowConvert)
-							.OnSelectionChanged_Lambda([this, SupportedBindingTypes, &SectionCategory, &DetailBuilder](TSharedPtr<FText> NewValue, ESelectInfo::Type)
-							{
-								int32 CurrentIndex = BindingTypeNames.IndexOfByPredicate([&](const TSharedPtr<FText>& In) {
-									return In == NewValue;
-									});
-
-								for(FMovieSceneBindingPropertyInfo& BindingPropertyInfo : BindingList->Bindings)
-								{
-									if (CurrentIndex == 0)
-									{
-										// Possessable- clear everything
-										BindingPropertyInfo.ResolveFlags = ELocatorResolveFlags::None;
-										BindingPropertyInfo.Locator = FUniversalObjectLocator();
-										BindingPropertyInfo.CustomBinding = nullptr;
-									}
-									else
-									{
-										// Custom binding
-										UClass* BindingTypeClass = SupportedBindingTypes[CurrentIndex - 1];
-										BindingPropertyInfo.Locator = FUniversalObjectLocator();
-										BindingPropertyInfo.ResolveFlags = ELocatorResolveFlags::None;
-										BindingPropertyInfo.CustomBinding = NewObject<UMovieSceneCustomBinding>(MovieScene.Get(), BindingTypeClass, NAME_None, EObjectFlags::RF_Transient);
-									}
-								}
-
-								// Force reset the view
-								DetailBuilder.GetPropertyUtilities()->NotifyFinishedChangingProperties(FPropertyChangedEvent(FindFProperty<FProperty>(FMovieSceneBindingPropertyInfo::StaticStruct(), "Locator")));
-								SectionCategory.GetParentLayout().ForceRefreshDetails();
-							})
-						.OnGenerateWidget_Lambda([](TSharedPtr<FText> Item)
-							{
-								return SNew(STextBlock).Text(*Item.Get());
-							})
-								.InitiallySelectedItem(BindingTypeNames[InitialIndex])
+							SNew(SComboButton)
+							.OnGetMenuContent(FOnGetContent::CreateLambda([this, &DetailBuilder]() { return GenerateBindingTypePicker(DetailBuilder);}))
+							.ContentPadding(0.0f)
+							.ButtonContent()
+							[
+								SNew(SHorizontalBox)
+								+ SHorizontalBox::Slot()
+								.AutoWidth()
+								.VAlign(VAlign_Center)
+								.Padding(0.0f, 0.0f, 4.0f, 0.0f)
+								[
+									SNew(SImage)
+									.Image(this, &FMovieSceneBindingPropertyInfoListCustomization::GetBindingTypeIcon)
+								]
+							+ SHorizontalBox::Slot()
+								.VAlign(VAlign_Center)
 								[
 									SNew(STextBlock)
-									.Margin(FMargin(0.0f, 2.0f))
-								.Text_Lambda([this, InitialIndex]() { return *BindingTypeNames[InitialIndex].Get(); })
+									.Text(this, &FMovieSceneBindingPropertyInfoListCustomization::GetBindingTypeValueAsString)
 								]
+							]
 						];
 
 					// Only show certain menus if we have a currently bound object
@@ -236,21 +213,97 @@ void FMovieSceneBindingPropertyInfoListCustomization::CustomizeDetails(IDetailLa
 	}
 }
 
-TSharedRef<SWidget> FMovieSceneBindingPropertyInfoListCustomization::OnGetConvertMenuContent(IDetailLayoutBuilder& DetailBuilder)
+TSharedRef<SWidget> FMovieSceneBindingPropertyInfoListCustomization::GenerateBindingTypePicker(IDetailLayoutBuilder& DetailBuilder)
 {
 	FMenuBuilder MenuBuilder(true, nullptr, nullptr, true);
+	if (SequencerPtr.IsValid())
+	{
+		TSharedRef<ISequencer> Sequencer = SequencerPtr.Pin().ToSharedRef();
+
+		TArray<FSequencerChangeBindingInfo> Bindings;
+		for (int32 BindingIndex = 0; BindingIndex < BindingList->Bindings.Num(); ++BindingIndex)
+		{
+			Bindings.Add({ BindingGuid, BindingIndex });
+		}
+		LevelSequenceEditorSubsystem->AddChangeBindingTypeMenu(MenuBuilder, Sequencer, Bindings, false, [this, &DetailBuilder]()
+			{
+				if (IDetailsView* DetailsView = DetailBuilder.GetDetailsView())
+				{
+					LevelSequenceEditorSubsystem->RefreshBindingDetails(DetailsView, BindingGuid);
+				}
+			});
+	}
+	return MenuBuilder.MakeWidget();
+}
 
+FText FMovieSceneBindingPropertyInfoListCustomization::GetBindingTypeValueAsString() const
+{
+	TSharedPtr<ISequencer> Sequencer = SequencerPtr.Pin();
+	if (BindingList)
+	{
+		if (Sequencer.IsValid())
+		{
+			UMovieSceneSequence* Sequence = MovieScene->GetTypedOuter<UMovieSceneSequence>();
+			if (BindingList->Bindings.Num() > 0)
+			{
+				// All bindings will be the same type
+				if (!BindingList->Bindings[0].CustomBinding)
+				{
+					// Possessable
+					return LOCTEXT("BindingType_Possessable", "Possessable");
+				}
+				else
+				{
+					 return BindingList->Bindings[0].CustomBinding->GetBindingTypePrettyName();
+				}
+			}
+		}
+	}
+	return FText();
+}
+
+const FSlateBrush* FMovieSceneBindingPropertyInfoListCustomization::GetBindingTypeIcon() const
+{
 	TSharedPtr<ISequencer> Sequencer = SequencerPtr.Pin();
+	if (BindingList)
+	{
+		if (Sequencer.IsValid())
+		{
+			UMovieSceneSequence* Sequence = MovieScene->GetTypedOuter<UMovieSceneSequence>();
+			if (BindingList->Bindings.Num() > 0)
+			{
+				// All bindings will be the same type
+				if (!BindingList->Bindings[0].CustomBinding)
+				{
+					// Possessable
+					return nullptr;
+				}
+				else
+				{
+					return BindingList->Bindings[0].CustomBinding->GetBindingTrackCustomIconOverlay().GetIcon();
+				}
+			}
+		}
+	}
+
+	return nullptr;
+}
 
-	if (Sequencer.IsValid() && BindingList)
+TSharedRef<SWidget> FMovieSceneBindingPropertyInfoListCustomization::OnGetConvertMenuContent(IDetailLayoutBuilder& DetailBuilder)
+{
+	FMenuBuilder MenuBuilder(true, nullptr, nullptr, true);
+
+
+	if (SequencerPtr.IsValid() && BindingList)
 	{
+		TSharedRef<ISequencer> Sequencer = SequencerPtr.Pin().ToSharedRef();
 		TArray<FSequencerChangeBindingInfo> Bindings;
 		for (int32 BindingIndex = 0; BindingIndex < BindingList->Bindings.Num(); ++BindingIndex)
 		{
 			Bindings.Add({ BindingGuid, BindingIndex });
 		}
 
-		FSequencerUtilities::AddConvertBindingMenu(MenuBuilder, Sequencer.ToSharedRef(), Bindings, [this, &DetailBuilder]()
+		LevelSequenceEditorSubsystem->AddChangeBindingTypeMenu(MenuBuilder, Sequencer, Bindings, true, [this, &DetailBuilder]()
 			{
 				if (IDetailsView* DetailsView = DetailBuilder.GetDetailsView())
 				{
@@ -551,9 +604,6 @@ void ULevelSequenceEditorSubsystem::Initialize(FSubsystemCollectionBase& Collect
 		}));
 
 	BindingPropertiesMenuExtender->AddMenuExtension("CustomBinding", EExtensionHook::First, CommandList, FMenuExtensionDelegate::CreateLambda([this](FMenuBuilder& MenuBuilder) {
-		// Only add menu entries where the focused sequence is a ULevelSequence
-
-		// Add instanced detail customizations
 
 		FFormatNamedArguments Args;
 		MenuBuilder.AddSubMenu(
@@ -562,6 +612,16 @@ void ULevelSequenceEditorSubsystem::Initialize(FSubsystemCollectionBase& Collect
 			FNewMenuDelegate::CreateLambda([this](FMenuBuilder& SubMenuBuilder) { AddBindingPropertiesMenu(SubMenuBuilder); }));
 		}));
 
+	BindingPropertiesMenuExtender->AddMenuExtension("ConvertBinding", EExtensionHook::First, CommandList, FMenuExtensionDelegate::CreateLambda([this](FMenuBuilder& MenuBuilder) {
+		FFormatNamedArguments Args;
+
+		MenuBuilder.AddSubMenu(
+			FText::Format(LOCTEXT("BindingProperties", "Binding Properties"), Args),
+			FText::Format(LOCTEXT("BindingPropertiesTooltip", "Modify the actor and object bindings for this track"), Args),
+			FNewMenuDelegate::CreateLambda([this](FMenuBuilder& SubMenuBuilder) 
+			{ AddConvertBindingsMenu(SubMenuBuilder); }));
+		}));
+
 	SequencerModule.GetObjectBindingContextMenuExtensibilityManager()->AddExtender(RebindComponentMenuExtender);
 
 	SidebarMenuExtender = MakeShared<FExtender>();
@@ -650,7 +710,7 @@ void ULevelSequenceEditorSubsystem::AddBindingDetailCustomizations(TSharedRef<ID
 					return MakeShared<FMovieSceneBindingPropertyInfoDetailCustomization>(InSequencer, InMovieScene, InBindingGuid, LevelSequenceEditorSubsystem);
 				}, ActiveSequencer.ToWeakPtr(), MovieScene, BindingGuid, this));
 
-			DetailsView->RegisterInstancedCustomPropertyTypeLayout(FMovieSceneDynamicBinding::StaticStruct()->GetFName(), FOnGetPropertyTypeCustomizationInstance::CreateStatic(&FMovieSceneDynamicBindingCustomization::MakeInstance, MovieScene, BindingGuid));
+			DetailsView->RegisterInstancedCustomPropertyTypeLayout(FMovieSceneDynamicBinding::StaticStruct()->GetFName(), FOnGetPropertyTypeCustomizationInstance::CreateStatic(&FMovieSceneDynamicBindingCustomization::MakeInstance, MovieScene, BindingGuid, 0));
 			
 			DetailsView->RegisterInstancedCustomPropertyLayout(UMovieSceneSpawnableActorBinding::StaticClass(), FOnGetDetailCustomizationInstance::CreateStatic(&FMovieSceneSpawnableActorBindingBaseCustomization::MakeInstance, ActiveSequencer.ToWeakPtr(), MovieScene, BindingGuid));
 		}
@@ -2504,6 +2564,7 @@ void ULevelSequenceEditorSubsystem::AddBindingPropertiesMenu(FMenuBuilder& MenuB
 			DetailsViewArgs.bShowScrollBar = false;
 			DetailsViewArgs.bAllowMultipleTopLevelObjects = false;
 			DetailsViewArgs.NotifyHook = &NotifyHook;
+			DetailsViewArgs.DefaultsOnlyVisibility = EEditDefaultsOnlyNodeVisibility::Automatic;
 		}
 
 		TSharedRef<IDetailsView> DetailsView = FModuleManager::GetModuleChecked<FPropertyEditorModule>("PropertyEditor").CreateDetailView(DetailsViewArgs);
@@ -2521,6 +2582,49 @@ void ULevelSequenceEditorSubsystem::AddBindingPropertiesMenu(FMenuBuilder& MenuB
 	}
 }
 
+void ULevelSequenceEditorSubsystem::AddConvertBindingsMenu(FMenuBuilder& MenuBuilder)
+{
+	// Binding conversion
+	
+	MenuBuilder.AddSubMenu(
+		LOCTEXT("ConvertBindingLabel", "Convert Selected Binding(s) To..."),
+		LOCTEXT("ConvertBindingLabelTooltip", "Convert selected bindings into another binding type"),
+		FNewMenuDelegate::CreateLambda([this](FMenuBuilder& MenuBuilder)
+			{
+				TSharedPtr<ISequencer> Sequencer = GetActiveSequencer();
+				if (Sequencer == nullptr)
+				{
+					return;
+				}
+
+				UMovieSceneSequence* const Sequence = Sequencer->GetFocusedMovieSceneSequence();
+				if (!IsValid(Sequence))
+				{
+					return;
+				}
+
+				TArray<FGuid> ObjectBindings;
+				Sequencer->GetSelectedObjects(ObjectBindings);
+				if (ObjectBindings.Num() == 0)
+				{
+					return;
+				}
+
+				TArray<FSequencerChangeBindingInfo> Bindings;
+				const FMovieSceneBindingReferences* BindingReferences = Sequence->GetBindingReferences();
+				for (FGuid ObjectGuid : ObjectBindings)
+				{
+					int32 BindingIndex = 0;
+					for (const FMovieSceneBindingReference& Reference : BindingReferences->GetReferences(ObjectGuid))
+					{
+						Bindings.Add({ Reference.ID, BindingIndex++ });
+					}
+				}
+
+				AddChangeBindingTypeMenu(MenuBuilder, Sequencer.ToSharedRef(), Bindings, true, TFunction<void()>());
+			}));
+}
+
 void ULevelSequenceEditorSubsystem::AddBindingPropertiesSidebar(FMenuBuilder& MenuBuilder)
 {
 	AddBindingPropertiesMenu(MenuBuilder);
@@ -3091,4 +3195,510 @@ void ULevelSequenceEditorSubsystem::RefreshTrackRowMetadataDetails(IDetailsView*
 	DetailsView->SetObjects(WeakHelpers, true);
 }
 
+
+void ULevelSequenceEditorSubsystem::AddChangeBindingTypeMenu(FMenuBuilder& MenuBuilder, TSharedRef<ISequencer> Sequencer, const TArray<FSequencerChangeBindingInfo>& BindingsToChange, bool bConvert, TFunction<void()> OnBindingChanged)
+{
+	UMovieSceneSequence* Sequence = Sequencer->GetFocusedMovieSceneSequence();
+	UMovieScene* MovieScene = Sequence ? Sequence->GetMovieScene() : nullptr;
+
+	if (!MovieScene || BindingsToChange.IsEmpty())
+	{
+		return;
+	}
+
+	// Can convert to possessable
+	if (!bConvert || Algo::AllOf(BindingsToChange, [this, Sequencer, Sequence, MovieScene](const FSequencerChangeBindingInfo& BindingInfo) {
+		return FSequencerUtilities::CanConvertToPossessable(Sequencer, BindingInfo.BindingID, BindingInfo.BindingIndex);
+		}))
+	{
+		const TWeakPtr<ISequencer> WeakSequencer = Sequencer;
+
+		MenuBuilder.AddMenuEntry(
+			LOCTEXT("ConvertToPossessable", "Possessable"),
+			bConvert ? LOCTEXT("ConvertToPossessableTooltip", "Convert selected binding(s) to a possessable") : LOCTEXT("ChangeToPossessableTooltip", "Reset selected binding(s) to a new possessable"),
+			FSlateIcon(),
+			FUIAction(FExecuteAction::CreateLambda([this, WeakSequencer, Sequence, BindingsToChange, bConvert, OnBindingChanged]()
+				{
+					const TSharedPtr<ISequencer> Sequencer = WeakSequencer.Pin();
+					if (!Sequencer)
+					{
+						return;
+					}
+
+					const TSharedRef<ISequencer> SequencerRef = Sequencer.ToSharedRef();
+
+					ChangeBindingTypes(SequencerRef, BindingsToChange, [SequencerRef, Sequence, bConvert](FGuid BindingID, int32 BindingIndex)
+						{
+							if (bConvert)
+							{
+								return FSequencerUtilities::ConvertToPossessable(SequencerRef, BindingID, BindingIndex);
+							}
+							else if (FMovieSceneBindingReferences* BindingReferences = Sequence->GetBindingReferences())
+							{
+								BindingReferences->AddOrReplaceBinding(BindingID, FUniversalObjectLocator(), BindingIndex);
+								return Sequence->GetMovieScene()->FindPossessable(BindingID);
+							}
+							return (FMovieScenePossessable*)nullptr;
+						},
+						OnBindingChanged);
+
+				})));
+	}
+
+
+	// Sort custom binding types by engine types vs. non-engine (custom user types)
+	TArrayView<const TSubclassOf<UMovieSceneCustomBinding>> PrioritySortedCustomBindingTypes = Sequencer->GetSupportedCustomBindingTypes();
+	TArray<const TSubclassOf<UMovieSceneCustomBinding>> EngineClasses;
+	TArray<const TSubclassOf<UMovieSceneCustomBinding>> UserClasses;
+	for (const TSubclassOf<UMovieSceneCustomBinding>& CustomBindingType : PrioritySortedCustomBindingTypes)
+	{
+		FString PackagePathName = CustomBindingType->GetPackage()->GetPathName();
+		if (PackagePathName.StartsWith(TEXT("/Engine")) || PackagePathName.StartsWith(TEXT("/Script")))
+		{
+			EngineClasses.Add(CustomBindingType);
+		}
+		else
+		{
+			UserClasses.Add(CustomBindingType);
+		}
+	}
+
+	// Show built-in classes first
+	for (const TSubclassOf<UMovieSceneCustomBinding>& CustomBindingType : EngineClasses)
+	{
+		FString BindingTypePath = CustomBindingType->GetPathName();
+		FString OtherPath = CustomBindingType->GetPackage()->GetPathName();
+
+		// Can convert to custom bindings
+		if (!bConvert || Algo::AllOf(BindingsToChange, [this, &Sequence, &MovieScene, CustomBindingType, Sequencer](const FSequencerChangeBindingInfo& BindingInfo)
+			{
+				return FSequencerUtilities::CanConvertToCustomBinding(Sequencer, BindingInfo.BindingID, CustomBindingType, BindingInfo.BindingIndex);
+			}))
+		{
+			const TWeakPtr<ISequencer> WeakSequencer = Sequencer;
+
+			// Special case director blueprint bindings to show sub-menu for setting endpoints
+			if (CustomBindingType == UMovieSceneSpawnableDirectorBlueprintBinding::StaticClass()
+				|| CustomBindingType == UMovieSceneReplaceableDirectorBlueprintBinding::StaticClass())
+			{
+				// Option to use a director blueprint condition and create or quick bind to an endpoint
+				MenuBuilder.AddSubMenu(
+					FText::Format(LOCTEXT("ConvertToDirectorBlueprintBinding", "{0}..."), CustomBindingType->GetDefaultObject<UMovieSceneCustomBinding>()->GetBindingTypePrettyName()),
+					bConvert ? FText::Format(LOCTEXT("ConvertToCustomBindingTooltip", "Convert selected binding to {0}"), CustomBindingType->GetDefaultObject<UMovieSceneCustomBinding>()->GetBindingTypePrettyName())
+					: FText::Format(LOCTEXT("ChangeToCustomBindingTooltip", "Reset selected binding to a new {0}"), CustomBindingType->GetDefaultObject<UMovieSceneCustomBinding>()->GetBindingTypePrettyName()),
+					FNewMenuDelegate::CreateLambda([this, Sequencer, BindingsToChange, bConvert, OnBindingChanged, CustomBindingType](FMenuBuilder& SubMenuBuilder) { FillDirectorBlueprintBindingSubMenu(SubMenuBuilder, Sequencer, BindingsToChange, bConvert, OnBindingChanged, CustomBindingType); }),
+					false,
+					CustomBindingType->GetDefaultObject<UMovieSceneCustomBinding>()->GetBindingTrackCustomIconOverlay()
+				);
+			}
+			else
+			{
+				MenuBuilder.AddMenuEntry(
+					CustomBindingType->GetDefaultObject<UMovieSceneCustomBinding>()->GetBindingTypePrettyName(),
+					bConvert ? FText::Format(LOCTEXT("ConvertToCustomBindingTooltip", "Convert selected binding to {0}"), CustomBindingType->GetDefaultObject<UMovieSceneCustomBinding>()->GetBindingTypePrettyName())
+					: FText::Format(LOCTEXT("ChangeToCustomBindingTooltip", "Reset selected binding to a new {0}"), CustomBindingType->GetDefaultObject<UMovieSceneCustomBinding>()->GetBindingTypePrettyName()),
+					CustomBindingType->GetDefaultObject<UMovieSceneCustomBinding>()->GetBindingTrackCustomIconOverlay(),
+					FUIAction(FExecuteAction::CreateLambda([this, WeakSequencer, CustomBindingType, Sequence, bConvert, BindingsToChange, OnBindingChanged]()
+						{
+							const TSharedPtr<ISequencer> Sequencer = WeakSequencer.Pin();
+							if (!Sequencer)
+							{
+								return;
+							}
+
+							const TSharedRef<ISequencer> SequencerRef = Sequencer.ToSharedRef();
+							ChangeBindingTypes(SequencerRef, BindingsToChange, [SequencerRef, Sequence, bConvert, CustomBindingType](FGuid BindingID, int32 BindingIndex)
+								{
+									if (bConvert)
+									{
+										return FSequencerUtilities::ConvertToCustomBinding(SequencerRef, BindingID, CustomBindingType, BindingIndex);
+									}
+									else if (FMovieSceneBindingReferences* BindingReferences = Sequence->GetBindingReferences())
+									{
+										BindingReferences->AddOrReplaceBinding(BindingID,
+											NewObject<UMovieSceneCustomBinding>(Sequence->GetMovieScene(), CustomBindingType),
+											BindingIndex);
+										return Sequence->GetMovieScene()->FindPossessable(BindingID);
+									}
+									return (FMovieScenePossessable*)nullptr;
+								},
+								OnBindingChanged);
+						})));
+			}
+		}
+	}
+
+	MenuBuilder.AddSeparator();
+
+	// Custom classes
+	MenuBuilder.AddSubMenu(
+		LOCTEXT("CustomBindings", "Custom Bindings..."),
+		LOCTEXT("CustomBindingsTooltip", "Choose or create a custom binding type"),
+		FNewMenuDelegate::CreateLambda([this, Sequencer, BindingsToChange, bConvert, OnBindingChanged, UserClasses](FMenuBuilder& SubMenuBuilder) { FillBindingClassSubMenu(SubMenuBuilder, Sequencer, BindingsToChange, bConvert, OnBindingChanged, UserClasses); }),
+		false
+	);
+
+}
+
+void ULevelSequenceEditorSubsystem::FillDirectorBlueprintBindingSubMenu(FMenuBuilder& MenuBuilder, TSharedRef<ISequencer> Sequencer, const TArray<FSequencerChangeBindingInfo>& BindingsToChange, bool bConvert, TFunction<void()> OnBindingChanged, const TSubclassOf<UMovieSceneCustomBinding>& CustomBindingType)
+{
+	UMovieSceneSequence* Sequence = Sequencer->GetFocusedMovieSceneSequence();
+	UMovieScene* MovieScene = Sequence ? Sequence->GetMovieScene() : nullptr;
+	const TWeakPtr<ISequencer> WeakSequencer = Sequencer;
+	if (Sequence)
+	{
+		MenuBuilder.AddMenuEntry(
+			LOCTEXT("CreateBindingEndpoint_Text", "Create New Binding Endpoint"),
+			LOCTEXT("CreateBindingEndpoint_Tooltip", "Creates a new binding endpoint in this sequence's blueprint."),
+			FSlateIcon(FAppStyle::GetAppStyleSetName(), "Sequencer.CreateEventBinding"),
+			FUIAction(
+				FExecuteAction::CreateLambda([this, WeakSequencer, CustomBindingType, Sequence, bConvert, BindingsToChange, OnBindingChanged]()
+					{
+						const TSharedPtr<ISequencer> Sequencer = WeakSequencer.Pin();
+						if (!Sequencer)
+						{
+							return;
+						}
+
+						const TSharedRef<ISequencer> SequencerRef = Sequencer.ToSharedRef();
+
+						if (Sequence)
+						{
+							// Change or convert the binding
+							ChangeBindingTypes(SequencerRef, BindingsToChange, [SequencerRef, Sequence, bConvert, CustomBindingType](FGuid BindingID, int32 BindingIndex)
+								{
+									FMovieScenePossessable* NewPossessable = nullptr;
+									if (FMovieSceneBindingReferences* BindingReferences = Sequence->GetBindingReferences())
+									{
+										if (bConvert)
+										{
+											NewPossessable = FSequencerUtilities::ConvertToCustomBinding(SequencerRef, BindingID, CustomBindingType, BindingIndex);
+										}
+										else
+										{
+											BindingReferences->AddOrReplaceBinding(BindingID,
+												NewObject<UMovieSceneCustomBinding>(Sequence->GetMovieScene(), CustomBindingType),
+												BindingIndex);
+											NewPossessable = Sequence->GetMovieScene()->FindPossessable(BindingID);
+										}
+
+										UMovieSceneCustomBinding* NewCustomBinding = BindingReferences->GetCustomBinding(BindingID, BindingIndex);
+
+										TArray<void*> RawData;
+										if (UMovieSceneReplaceableDirectorBlueprintBinding* ReplaceableBinding = Cast<UMovieSceneReplaceableDirectorBlueprintBinding>(NewCustomBinding))
+										{
+											RawData.Add(&ReplaceableBinding->DynamicBinding);
+										}
+										else if (UMovieSceneSpawnableDirectorBlueprintBinding* SpawnableBinding = Cast<UMovieSceneSpawnableDirectorBlueprintBinding>(NewCustomBinding))
+										{
+											RawData.Add(&SpawnableBinding->DynamicBinding);
+										}
+
+										// Create temporary director blueprint binding customization for use in creating the endpoint
+										TSharedRef<FMovieSceneDynamicBindingCustomization> BlueprintBindingCustomization = StaticCastSharedRef<FMovieSceneDynamicBindingCustomization>(FMovieSceneDynamicBindingCustomization::MakeInstance(Sequence->GetMovieScene(), BindingID, BindingIndex));
+										BlueprintBindingCustomization->SetRawData(RawData);
+										BlueprintBindingCustomization->CreateEndpoint();
+									}
+
+									return NewPossessable;
+								},
+								OnBindingChanged);
+						}
+					}
+				)
+			));
+
+		MenuBuilder.AddSubMenu(
+			LOCTEXT("CreateQuickBinding_Text", "Quick Bind"),
+			LOCTEXT("CreateQuickBinding_Tooltip", "Shows a list of functions in this sequence's blueprint that can be used for bindings."),
+			FNewMenuDelegate::CreateLambda([this, Sequencer, BindingsToChange, bConvert, OnBindingChanged, CustomBindingType](FMenuBuilder& SubMenuBuilder) { PopulateQuickBindSubMenu(SubMenuBuilder, Sequencer, BindingsToChange, bConvert, OnBindingChanged, CustomBindingType); }),
+			false /* bInOpenSubMenuOnClick */,
+			FSlateIcon(FAppStyle::GetAppStyleSetName(), "Sequencer.CreateQuickBinding"),
+			false /* bInShouldWindowAfterMenuSelection */
+		);
+	}
+}
+
+void ULevelSequenceEditorSubsystem::PopulateQuickBindSubMenu(FMenuBuilder& MenuBuilder, TSharedRef<ISequencer> Sequencer, const TArray<FSequencerChangeBindingInfo>& BindingsToChange, bool bConvert, TFunction<void()> OnBindingChanged, const TSubclassOf<UMovieSceneCustomBinding>& CustomBindingType)
+{
+	UMovieSceneSequence* Sequence = Sequencer->GetFocusedMovieSceneSequence();
+	UMovieScene* MovieScene = Sequence ? Sequence->GetMovieScene() : nullptr;
+
+	if (!MovieScene || BindingsToChange.IsEmpty())
+	{
+		return;
+	}
+
+	TSharedRef<FMovieSceneDynamicBindingCustomization> BlueprintBindingCustomization = StaticCastSharedRef<FMovieSceneDynamicBindingCustomization>(FMovieSceneDynamicBindingCustomization::MakeInstance(Sequence->GetMovieScene(), BindingsToChange[0].BindingID, BindingsToChange[0].BindingIndex));
+
+	BlueprintBindingCustomization->PopulateQuickBindSubMenu(MenuBuilder,
+		Sequence,
+		FOnQuickBindActionSelected::CreateLambda([this, Sequencer, BindingsToChange, Sequence, bConvert, OnBindingChanged, CustomBindingType]
+		(
+			const TArray<TSharedPtr<FEdGraphSchemaAction>>& SelectedAction,
+			ESelectInfo::Type InSelectionType,
+			UBlueprint* Blueprint,
+			FMovieSceneDirectorBlueprintEndpointDefinition EndpointDefinition
+			)
+			{
+				// Change or convert the binding
+				ChangeBindingTypes(Sequencer, BindingsToChange, [Sequencer, Sequence, bConvert, OnBindingChanged, CustomBindingType, &SelectedAction, InSelectionType, Blueprint, EndpointDefinition](FGuid BindingID, int32 BindingIndex)
+					{
+						FMovieScenePossessable* NewPossessable = nullptr;
+						if (FMovieSceneBindingReferences* BindingReferences = Sequence->GetBindingReferences())
+						{
+							if (bConvert)
+							{
+								NewPossessable = FSequencerUtilities::ConvertToCustomBinding(Sequencer, BindingID, CustomBindingType, BindingIndex);
+							}
+							else
+							{
+								BindingReferences->AddOrReplaceBinding(BindingID,
+									NewObject<UMovieSceneCustomBinding>(Sequence->GetMovieScene(), CustomBindingType),
+									BindingIndex);
+								NewPossessable = Sequence->GetMovieScene()->FindPossessable(BindingID);
+							}
+
+							UMovieSceneCustomBinding* NewCustomBinding = BindingReferences->GetCustomBinding(BindingID, BindingIndex);
+
+							TArray<void*> RawData;
+							if (UMovieSceneReplaceableDirectorBlueprintBinding* ReplaceableBinding = Cast<UMovieSceneReplaceableDirectorBlueprintBinding>(NewCustomBinding))
+							{
+								RawData.Add(&ReplaceableBinding->DynamicBinding);
+							}
+							else if (UMovieSceneSpawnableDirectorBlueprintBinding* SpawnableBinding = Cast<UMovieSceneSpawnableDirectorBlueprintBinding>(NewCustomBinding))
+							{
+								RawData.Add(&SpawnableBinding->DynamicBinding);
+							}
+
+							// Create temporary director blueprint binding customization for use in creating the endpoint
+							TSharedRef<FMovieSceneDynamicBindingCustomization> BlueprintBindingCustomization = StaticCastSharedRef<FMovieSceneDynamicBindingCustomization>(FMovieSceneDynamicBindingCustomization::MakeInstance(Sequence->GetMovieScene(), BindingID, BindingIndex));
+							BlueprintBindingCustomization->SetRawData(RawData);
+							BlueprintBindingCustomization->HandleQuickBindActionSelected(SelectedAction, InSelectionType, Blueprint, EndpointDefinition);
+
+						}
+						return NewPossessable;
+					},
+					OnBindingChanged);
+			}));
+}
+
+void ULevelSequenceEditorSubsystem::FillBindingClassSubMenu(FMenuBuilder& MenuBuilder, TSharedRef<ISequencer> Sequencer, const TArray<FSequencerChangeBindingInfo>& BindingsToChange, bool bConvert, TFunction<void()> OnBindingsChanged, const TArray<const TSubclassOf<UMovieSceneCustomBinding>>& UserCustomBindingTypes)
+{
+	UMovieSceneSequence* Sequence = Sequencer->GetFocusedMovieSceneSequence();
+	UMovieScene* MovieScene = Sequence ? Sequence->GetMovieScene() : nullptr;
+
+	if (!MovieScene || BindingsToChange.IsEmpty())
+	{
+		return;
+	}
+
+	// If we allow creating binding classes from this replaceable actor binding...
+	if (MovieScene->IsCustomBindingClassAllowed(UMovieSceneReplaceableActorBinding_BPBase::StaticClass()))
+	{
+		// Create a new custom binding Class
+		MenuBuilder.AddMenuEntry(
+			LOCTEXT("CreateNewCustomBindingClass", "Create new Replaceable Binding Class"),
+			LOCTEXT("CreateNewCustomBindingClassTooltip", "Creates a new replaceable binding blueprint class"),
+			UMovieSceneReplaceableActorBinding_BPBase::StaticClass()->GetDefaultObject<UMovieSceneReplaceableActorBinding_BPBase>()->GetBindingTrackCustomIconOverlay(),
+			FUIAction(FExecuteAction::CreateLambda([this, Sequencer, Sequence, bConvert, BindingsToChange, OnBindingsChanged]()
+				{
+					FAssetToolsModule& AssetToolsModule = FModuleManager::GetModuleChecked<FAssetToolsModule>("AssetTools");
+
+					FString NewBindingPath = Sequence->GetPathName();
+					FString NewBindingName = Sequence->GetName() + TEXT("_BindingType");
+					AssetToolsModule.Get().CreateUniqueAssetName(NewBindingPath + TEXT("/") + NewBindingName, TEXT(""), NewBindingPath, NewBindingName);
+
+					UBlueprint* Blueprint = FKismetEditorUtilities::CreateBlueprintFromClass(LOCTEXT("CreateNewBindingClass", "Create New Replaceable Binding Class"), UMovieSceneReplaceableActorBinding_BPBase::StaticClass(), NewBindingName);
+
+					if (Blueprint != NULL && Blueprint->GeneratedClass)
+					{
+						GEditor->GetEditorSubsystem<UAssetEditorSubsystem>()->OpenEditorForAsset(Blueprint);
+
+						// Implement the ResolveRuntimeBindingInternal function
+
+						UFunction* OverrideFunc = FindUField<UFunction>(UMovieSceneReplaceableActorBinding_BPBase::StaticClass(), GET_FUNCTION_NAME_CHECKED(UMovieSceneReplaceableActorBinding_BPBase, BP_ResolveRuntimeBinding));
+						check(OverrideFunc);
+						Blueprint->Modify();
+						// Implement the function graph
+						UEdGraph* const NewGraph = FBlueprintEditorUtils::CreateNewGraph(Blueprint, TEXT("BP_ResolveRuntimeBinding"), UEdGraph::StaticClass(), UEdGraphSchema_K2::StaticClass());
+						FBlueprintEditorUtils::AddFunctionGraph(Blueprint, NewGraph, /*bIsUserCreated=*/ false, UMovieSceneReplaceableActorBinding_BPBase::StaticClass());
+						NewGraph->Modify();
+						FKismetEditorUtilities::CompileBlueprint(Blueprint);
+
+						// Change or convert the binding
+						ChangeBindingTypes(Sequencer, BindingsToChange, [Blueprint, Sequencer, Sequence, bConvert, OnBindingsChanged](FGuid BindingID, int32 BindingIndex)
+							{
+								FMovieScenePossessable* NewPossessable = nullptr;
+								if (FMovieSceneBindingReferences* BindingReferences = Sequence->GetBindingReferences())
+								{
+									if (bConvert)
+									{
+										NewPossessable = FSequencerUtilities::ConvertToCustomBinding(Sequencer, BindingID, TSubclassOf<UMovieSceneCustomBinding>(Blueprint->GeneratedClass), BindingIndex);
+									}
+									else
+									{
+										BindingReferences->AddOrReplaceBinding(BindingID,
+											NewObject<UMovieSceneCustomBinding>(Sequence->GetMovieScene(), Blueprint->GeneratedClass),
+											BindingIndex);
+										NewPossessable = Sequence->GetMovieScene()->FindPossessable(BindingID);
+									}
+								}
+								return NewPossessable;
+							},
+							OnBindingsChanged);
+
+
+						FKismetEditorUtilities::BringKismetToFocusAttentionOnObject(NewGraph);
+						Sequencer->RefreshSupportedCustomBindingTypes();
+					}
+				}
+		)));
+
+		class FUserCustomBindingFilter : public IClassViewerFilter
+		{
+		public:
+
+			TArray<const TSubclassOf<UMovieSceneCustomBinding>> UserCustomBindingTypes;
+
+			virtual bool IsClassAllowed(const FClassViewerInitializationOptions& InInitOptions, const UClass* InClass, TSharedRef< FClassViewerFilterFuncs > InFilterFuncs) override
+			{
+				return UserCustomBindingTypes.Contains(InClass);
+			}
+
+			virtual bool IsUnloadedClassAllowed(const FClassViewerInitializationOptions& InInitOptions, const TSharedRef< const IUnloadedBlueprintData > InBlueprint, TSharedRef< FClassViewerFilterFuncs > InFilterFuncs) override
+			{
+				return InBlueprint->GetNativeParent() && InBlueprint->GetNativeParent()->IsChildOf(UMovieSceneReplaceableActorBinding_BPBase::StaticClass());
+			}
+		};
+
+		MenuBuilder.BeginSection(TEXT("ChooseCustomBindingClass"), LOCTEXT("ChooseCustomBindingClass", "Choose Custom Binding Class"));
+		{
+			TSharedPtr<FUserCustomBindingFilter> BindingClassFilter = MakeShared<FUserCustomBindingFilter>();
+			BindingClassFilter->UserCustomBindingTypes = UserCustomBindingTypes;
+
+			FClassViewerInitializationOptions Options;
+			Options.bShowBackgroundBorder = false;
+			Options.bShowUnloadedBlueprints = true;
+			Options.NameTypeToDisplay = EClassViewerNameTypeToDisplay::DisplayName;
+
+			Options.ClassFilters.Add(BindingClassFilter.ToSharedRef());
+
+			MenuBuilder.AddWidget(FModuleManager::LoadModuleChecked<FClassViewerModule>("ClassViewer").CreateClassViewer(Options,
+				FOnClassPicked::CreateLambda([this, Sequencer, Sequence, bConvert, BindingsToChange, OnBindingsChanged](UClass* Class)
+					{
+						// Change or convert the binding
+						ChangeBindingTypes(Sequencer, BindingsToChange, [Sequencer, Sequence, Class, bConvert, OnBindingsChanged](FGuid BindingID, int32 BindingIndex)
+							{
+								FMovieScenePossessable* NewPossessable = nullptr;
+								if (FMovieSceneBindingReferences* BindingReferences = Sequence->GetBindingReferences())
+								{
+									if (bConvert)
+									{
+										NewPossessable = FSequencerUtilities::ConvertToCustomBinding(Sequencer, BindingID, TSubclassOf<UMovieSceneCustomBinding>(Class), BindingIndex);
+									}
+									else
+									{
+										BindingReferences->AddOrReplaceBinding(BindingID,
+											NewObject<UMovieSceneCustomBinding>(Sequence->GetMovieScene(), Class),
+											BindingIndex);
+										NewPossessable = Sequence->GetMovieScene()->FindPossessable(BindingID);
+									}
+								}
+								return NewPossessable;
+							},
+							OnBindingsChanged);
+						FSlateApplication::Get().DismissAllMenus();
+					})), FText::GetEmpty(), true);
+		}
+		MenuBuilder.EndSection();
+	}
+}
+
+void ULevelSequenceEditorSubsystem::ChangeBindingTypes(const TSharedRef<ISequencer>& InSequencer
+	, const TArray<FSequencerChangeBindingInfo>& InBindingsToChange
+	, TFunction<FMovieScenePossessable* (FGuid, int32)> InDoChangeType
+	, TFunction<void()> InOnBindingChanged)
+{
+	using namespace UE::Sequencer;
+
+	UMovieSceneSequence* const Sequence = InSequencer->GetFocusedMovieSceneSequence();
+	if (!IsValid(Sequence))
+	{
+		return;
+	}
+
+	UMovieScene* const MovieScene = Sequence->GetMovieScene();
+	if (!IsValid(MovieScene))
+	{
+		return;
+	}
+
+	if (InBindingsToChange.Num() == 0)
+	{
+		return;
+	}
+
+	if (MovieScene->IsReadOnly())
+	{
+		FSequencerUtilities::ShowReadOnlyError();
+		return;
+	}
+
+	const FScopedTransaction Transaction(LOCTEXT("ConvertSelectedNodes", "Convert Selected Nodes Binding Type"));
+	MovieScene->Modify();
+
+	FScopedSlowTask SlowTask(InBindingsToChange.Num(), LOCTEXT("ConvertProgress", "Converting Selected Nodes Binding Type"));
+	SlowTask.MakeDialog(true);
+
+	TArray<AActor*> PossessedActors;
+	for (const FSequencerChangeBindingInfo& BindingInfo : InBindingsToChange)
+	{
+		SlowTask.EnterProgressFrame();
+
+		if (FMovieScenePossessable* Possessable = InDoChangeType(BindingInfo.BindingID, BindingInfo.BindingIndex))
+		{
+			InSequencer->ForceEvaluate();
+
+			for (TWeakObjectPtr<> WeakObject : InSequencer->FindBoundObjects(Possessable->GetGuid(), InSequencer->GetFocusedTemplateID()))
+			{
+				if (AActor* PossessedActor = Cast<AActor>(WeakObject.Get()))
+				{
+					PossessedActors.Add(PossessedActor);
+				}
+			}
+
+			if (GWarn->ReceivedUserCancel())
+			{
+				break;
+			}
+		}
+	}
+
+	if (PossessedActors.Num())
+	{
+		const bool bNotifySelectionChanged = true;
+		const bool bDeselectBSP = true;
+		const bool bWarnAboutTooManyActors = false;
+		const bool bSelectEvenIfHidden = false;
+
+		GEditor->GetSelectedActors()->Modify();
+		GEditor->GetSelectedActors()->BeginBatchSelectOperation();
+		GEditor->SelectNone(bNotifySelectionChanged, bDeselectBSP, bWarnAboutTooManyActors);
+		for (auto PossessedActor : PossessedActors)
+		{
+			GEditor->SelectActor(PossessedActor, true, bNotifySelectionChanged, bSelectEvenIfHidden);
+		}
+		GEditor->GetSelectedActors()->EndBatchSelectOperation();
+		GEditor->NoteSelectionChange();
+
+		InSequencer->NotifyMovieSceneDataChanged(EMovieSceneDataChangeType::MovieSceneStructureItemsChanged);
+	}
+
+	if (InOnBindingChanged != nullptr)
+	{
+		InOnBindingChanged();
+	}
+}
+
 #undef LOCTEXT_NAMESPACE
\ No newline at end of file
diff --git a/Engine/Plugins/MovieScene/LevelSequenceEditor/Source/LevelSequenceEditor/Private/Misc/LevelSequenceCustomization.cpp b/Engine/Plugins/MovieScene/LevelSequenceEditor/Source/LevelSequenceEditor/Private/Misc/LevelSequenceCustomization.cpp
--- a/Engine/Plugins/MovieScene/LevelSequenceEditor/Source/LevelSequenceEditor/Private/Misc/LevelSequenceCustomization.cpp
+++ b/Engine/Plugins/MovieScene/LevelSequenceEditor/Source/LevelSequenceEditor/Private/Misc/LevelSequenceCustomization.cpp
@@ -235,37 +235,9 @@ void FLevelSequenceCustomization::ExtendObjectBindingContextMenu(FMenuBuilder& M
 	
 	if (bShowConvert)
 	{
-		// Binding conversion
-		MenuBuilder.AddSubMenu(
-			LOCTEXT("ConvertBindingLabel", "Convert Selected Binding(s) To..."),
-			LOCTEXT("ConvertBindingLabelTooltip", "Convert selected bindings into another binding type"),
-			FNewMenuDelegate::CreateLambda([this](FMenuBuilder& MenuBuilder) 
-			{
-				const TSharedPtr<ISequencer> Sequencer = WeakSequencer.Pin();
-				if (!Sequencer.IsValid())
-				{
-					return;
-				}
-
-				UMovieSceneSequence* const Sequence = Sequencer->GetFocusedMovieSceneSequence();
-				if (!IsValid(Sequence))
-				{
-					return;
-				}
-
-				TArray<FSequencerChangeBindingInfo> Bindings;
-				const FMovieSceneBindingReferences* BindingReferences = Sequence->GetBindingReferences();
-				for (TViewModelPtr<IObjectBindingExtension> ObjectBindingNode : Sequencer->GetViewModel()->GetSelection()->Outliner.Filter<IObjectBindingExtension>())
-				{
-					int32 BindingIndex = 0;
-					for (const FMovieSceneBindingReference& Reference : BindingReferences->GetReferences(ObjectBindingNode->GetObjectGuid()))
-					{
-						Bindings.Add({ Reference.ID, BindingIndex++ });
-					}
-				}
-
-				FSequencerUtilities::AddConvertBindingMenu(MenuBuilder, Sequencer.ToSharedRef(), Bindings, TFunction<void()>()); 
-			}));
+		// We don't add anything here, but the extension will
+		MenuBuilder.BeginSection("ConvertBinding");
+		MenuBuilder.EndSection();
 	}
 
 	MenuBuilder.BeginSection("Import/Export", LOCTEXT("ImportExportMenuSectionName", "Import/Export"));
diff --git a/Engine/Plugins/MovieScene/LevelSequenceEditor/Source/LevelSequenceEditor/Public/LevelSequenceEditorSubsystem.h b/Engine/Plugins/MovieScene/LevelSequenceEditor/Source/LevelSequenceEditor/Public/LevelSequenceEditorSubsystem.h
--- a/Engine/Plugins/MovieScene/LevelSequenceEditor/Source/LevelSequenceEditor/Public/LevelSequenceEditorSubsystem.h
+++ b/Engine/Plugins/MovieScene/LevelSequenceEditor/Source/LevelSequenceEditor/Public/LevelSequenceEditorSubsystem.h
@@ -23,6 +23,8 @@ struct FMovieScenePasteFoldersParams;
 struct FMovieScenePasteSectionsParams;
 struct FMovieScenePasteTracksParams;
 struct FBakingAnimationKeySettings;
+struct FSequencerChangeBindingInfo;
+struct FMovieScenePossessable;
 
 DECLARE_LOG_CATEGORY_EXTERN(LogLevelSequenceEditor, Log, All);
 
@@ -329,10 +331,25 @@ private:
 
 	void AddAssignActorMenu(FMenuBuilder& MenuBuilder);
 	void AddBindingPropertiesMenu(FMenuBuilder& MenuBuilder);
+	void AddConvertBindingsMenu(FMenuBuilder& MenuBuilder);
+
+	
+	void FillDirectorBlueprintBindingSubMenu(FMenuBuilder& MenuBuilder, TSharedRef<ISequencer> Sequencer, const TArray<FSequencerChangeBindingInfo>& BindingsToChange, bool bConvert, TFunction<void()> OnBindingChanged, const TSubclassOf<UMovieSceneCustomBinding>& CustomBindingType);
+	void PopulateQuickBindSubMenu(FMenuBuilder& MenuBuilder, TSharedRef<ISequencer> Sequencer, const TArray<FSequencerChangeBindingInfo>& BindingsToChange, bool bConvert, TFunction<void()> OnBindingChanged, const TSubclassOf<UMovieSceneCustomBinding>& CustomBindingType);
+	void FillBindingClassSubMenu(FMenuBuilder& MenuBuilder, TSharedRef<ISequencer> Sequencer, const TArray<FSequencerChangeBindingInfo>& BindingsToChange, bool bConvert, TFunction<void()> OnBindingChanged, const TArray<const TSubclassOf<UMovieSceneCustomBinding>>& UserCustomBindingTypes);
+	void ChangeBindingTypes(const TSharedRef<ISequencer>& InSequencer
+		, const TArray<FSequencerChangeBindingInfo>& InBindingsToChange
+		, TFunction<FMovieScenePossessable* (FGuid, int32)> InDoChangeType
+		, TFunction<void()> InOnBindingChanged);
 
 	void AddTrackRowMetadataMenu(FMenuBuilder& MenuBuilder);
 public:
-	void AddBindingPropertiesSidebar(FMenuBuilder& MenuBuilder);
+	void AddBindingPropertiesSidebar(FMenuBuilder& MenuBuilder); 
+	
+	/* Creates a menu for changing or converting a binding type. If bConvert is true, it will only show types that state they are able to be converted to from the passed in bindings
+	and will attempt to convert them. If bConvert is false, it will change the binding type and reset to a default binding of that type.*/
+	void AddChangeBindingTypeMenu(FMenuBuilder& MenuBuilder, TSharedRef<ISequencer> Sequencer, const TArray<FSequencerChangeBindingInfo>& BindingsToChange, bool bConvert, TFunction<void()> OnBindingChanged);
+
 private:
 	void OnFinishedChangingLocators(const FPropertyChangedEvent& PropertyChangedEvent, TSharedRef<IDetailsView> DetailsView, FGuid ObjectBindingID);
 
diff --git a/Engine/Source/Editor/MovieSceneTools/Private/MovieSceneDirectorBlueprintEndpointCustomization.cpp b/Engine/Source/Editor/MovieSceneTools/Private/MovieSceneDirectorBlueprintEndpointCustomization.cpp
--- a/Engine/Source/Editor/MovieSceneTools/Private/MovieSceneDirectorBlueprintEndpointCustomization.cpp
+++ b/Engine/Source/Editor/MovieSceneTools/Private/MovieSceneDirectorBlueprintEndpointCustomization.cpp
@@ -95,6 +95,9 @@ void FMovieSceneDirectorBlueprintEndpointCustomization::CustomizeChildren(TShare
 	PropertyHandle = InPropertyHandle;
 	PropertyUtilities = CustomizationUtils.GetPropertyUtilities();
 
+	PropertyRawData.Empty();
+	PropertyHandle->AccessRawData(PropertyRawData);
+
 	ChildBuilder.AddCustomRow(FText())
 	.NameContent()
 	[
@@ -268,23 +271,20 @@ void FMovieSceneDirectorBlueprintEndpointCustomization::CustomizeChildren(TShare
 		TArray<UObject*> EditObjects;
 		GetEditObjects(EditObjects);
 
-		TArray<void*> RawData;
-		PropertyHandle->AccessRawData(RawData);
-
 		FPayloadVariableMap PayloadVariables;
 		TArray<FName> WellKnownParameters;
 		TArray<FName, TInlineAllocator<8>> AllValidNames;
 
 		const FString* WorldContextParamName = CommonFunction->FindMetaData(FBlueprintMetadata::MD_WorldContext);
 
-		for (int32 Index = 0; Index < RawData.Num(); ++Index)
+		for (int32 Index = 0; Index < PropertyRawData.Num(); ++Index)
 		{
 			AllValidNames.Empty();
 			PayloadVariables.Empty();
 			WellKnownParameters.Empty();
 
-			GetPayloadVariables(EditObjects[Index], RawData[Index], PayloadVariables);
-			GetWellKnownParameterPinNames(EditObjects[Index], RawData[Index], WellKnownParameters);
+			GetPayloadVariables(EditObjects[Index], PropertyRawData[Index], PayloadVariables);
+			GetWellKnownParameterPinNames(EditObjects[Index], PropertyRawData[Index], WellKnownParameters);
 
 			TSharedPtr<FStructOnScope> StructData = MakeShared<FStructOnScope>(CommonFunction);
 
@@ -428,17 +428,14 @@ void FMovieSceneDirectorBlueprintEndpointCustomization::OnPayloadVariableChanged
 	TArray<UObject*> EditObjects;
 	GetEditObjects(EditObjects);
 
-	TArray<void*> RawData;
-	PropertyHandle->AccessRawData(RawData);
-
-	if (!ensure(RawData.Num() == EditObjects.Num()))
+	if (!ensure(PropertyRawData.Num() == EditObjects.Num()))
 	{
 		return;
 	}
 
-	for (int32 Index = 0; Index < RawData.Num(); ++Index)
+	for (int32 Index = 0; Index < PropertyRawData.Num(); ++Index)
 	{
-		bChangedAnything |= SetPayloadVariable(EditObjects[Index], RawData[Index], Property->GetFName(), NewPayloadVariable);
+		bChangedAnything |= SetPayloadVariable(EditObjects[Index], PropertyRawData[Index], Property->GetFName(), NewPayloadVariable);
 	}
 
 	if (bChangedAnything)
@@ -524,7 +521,10 @@ void FMovieSceneDirectorBlueprintEndpointCustomization::OnSetCallInEditorCheckSt
 
 void FMovieSceneDirectorBlueprintEndpointCustomization::OnBlueprintCompiled(UBlueprint*)
 {
-	PropertyUtilities->ForceRefresh();
+	if (PropertyUtilities.IsValid())
+	{
+		PropertyUtilities->ForceRefresh();
+	}
 }
 
 UMovieSceneSequence* FMovieSceneDirectorBlueprintEndpointCustomization::GetCommonSequence() const
@@ -552,17 +552,14 @@ void FMovieSceneDirectorBlueprintEndpointCustomization::IterateEndpoints(TFuncti
 	TArray<UObject*> EditObjects;
 	GetEditObjects(EditObjects);
 
-	TArray<void*> RawData;
-	PropertyHandle->AccessRawData(RawData);
-
-	if (!ensure(RawData.Num() == EditObjects.Num()))
+	if (!ensure(PropertyRawData.Num() == EditObjects.Num()))
 	{
 		return;
 	}
 
-	for (int32 Index = 0; Index < RawData.Num(); ++Index)
+	for (int32 Index = 0; Index < PropertyRawData.Num(); ++Index)
 	{
-		UK2Node* Endpoint = FindEndpoint(EditObjects[Index], RawData[Index]);
+		UK2Node* Endpoint = FindEndpoint(EditObjects[Index], PropertyRawData[Index]);
 		if (Endpoint)
 		{
 			if (!Callback(Endpoint))
@@ -614,7 +611,10 @@ UK2Node* FMovieSceneDirectorBlueprintEndpointCustomization::GetCommonEndpoint()
 
 void FMovieSceneDirectorBlueprintEndpointCustomization::GetEditObjects(TArray<UObject*>& OutObjects) const
 {
-	PropertyHandle->GetOuterObjects(OutObjects);
+	if (PropertyHandle.IsValid())
+	{
+		PropertyHandle->GetOuterObjects(OutObjects);
+	}
 }
 
 TSharedRef<SWidget> FMovieSceneDirectorBlueprintEndpointCustomization::GetMenuContent()
@@ -719,16 +719,13 @@ bool FMovieSceneDirectorBlueprintEndpointCustomization::CompareWellKnownParamete
 	TArray<UObject*> EditObjects;
 	GetEditObjects(EditObjects);
 
-	TArray<void*> RawData;
-	PropertyHandle->AccessRawData(RawData);
-
-	check(EditObjects.Num() == RawData.Num());
+	check(EditObjects.Num() == PropertyRawData.Num());
 
 	TArray<FName> WellKnownParameters;
-	for (int32 Index = 0; Index < RawData.Num(); ++Index)
+	for (int32 Index = 0; Index < PropertyRawData.Num(); ++Index)
 	{
 		WellKnownParameters.Empty();
-		GetWellKnownParameterPinNames(EditObjects[Index], RawData[Index], WellKnownParameters);
+		GetWellKnownParameterPinNames(EditObjects[Index], PropertyRawData[Index], WellKnownParameters);
 		if (WellKnownParameters.IsValidIndex(ParameterIndex) &&
 				WellKnownParameters[ParameterIndex] == InPinName)
 		{
@@ -746,22 +743,22 @@ void FMovieSceneDirectorBlueprintEndpointCustomization::SetWellKnownParameterPin
 	TArray<UObject*> EditObjects;
 	GetEditObjects(EditObjects);
 
-	TArray<void*> RawData;
-	PropertyHandle->AccessRawData(RawData);
+	check(EditObjects.Num() == PropertyRawData.Num());
 
-	check(EditObjects.Num() == RawData.Num());
-
-	for (int32 Index = 0; Index < RawData.Num(); ++Index)
+	for (int32 Index = 0; Index < PropertyRawData.Num(); ++Index)
 	{
 		EditObjects[Index]->Modify();
 
-		const bool bWasSet = SetWellKnownParameterPinName(EditObjects[Index], RawData[Index], ParameterIndex, InNewBoundPinName);
+		const bool bWasSet = SetWellKnownParameterPinName(EditObjects[Index], PropertyRawData[Index], ParameterIndex, InNewBoundPinName);
 		ensureMsgf(bWasSet, TEXT("Trying to set bound pin %s for parameter %d but sub-class did not handle it."), 
 				*InNewBoundPinName.ToString(), ParameterIndex);
 	}
 
-	// Ensure that anything listening for property changed notifications are notified of the new binding
-	PropertyHandle->NotifyFinishedChangingProperties();
+	if (PropertyHandle.IsValid())
+	{
+		// Ensure that anything listening for property changed notifications are notified of the new binding
+		PropertyHandle->NotifyFinishedChangingProperties();
+	}
 
 	FSequenceDataMap AllSequenceData;
 	GatherSequenceData(AllSequenceData);
@@ -770,7 +767,10 @@ void FMovieSceneDirectorBlueprintEndpointCustomization::SetWellKnownParameterPin
 		FKismetEditorUtilities::CompileBlueprint(SequenceData.Value.Blueprint);
 	}
 
-	PropertyUtilities->ForceRefresh();
+	if (PropertyUtilities.IsValid())
+	{
+		PropertyUtilities->ForceRefresh();
+	}
 }
 
 const FSlateBrush* FMovieSceneDirectorBlueprintEndpointCustomization::GetWellKnownParameterPinIcon(int32 ParameterIndex) const
@@ -789,18 +789,15 @@ FText FMovieSceneDirectorBlueprintEndpointCustomization::GetWellKnownParameterPi
 	TArray<UObject*> EditObjects;
 	GetEditObjects(EditObjects);
 
-	TArray<void*> RawData;
-	PropertyHandle->AccessRawData(RawData);
-
-	check(EditObjects.Num() == RawData.Num());
+	check(EditObjects.Num() == PropertyRawData.Num());
 
 	FName CommonPinName;
 	TArray<FName> WellKnownParameters;
 
-	for (int32 Index = 0; Index < RawData.Num(); ++Index)
+	for (int32 Index = 0; Index < PropertyRawData.Num(); ++Index)
 	{
 		WellKnownParameters.Empty();
-		GetWellKnownParameterPinNames(EditObjects[Index], RawData[Index], WellKnownParameters);
+		GetWellKnownParameterPinNames(EditObjects[Index], PropertyRawData[Index], WellKnownParameters);
 		if (WellKnownParameters.IsValidIndex(ParameterIndex))
 		{
 			FName PinName = WellKnownParameters[ParameterIndex];
@@ -1245,9 +1242,7 @@ const FSlateBrush* FMovieSceneDirectorBlueprintEndpointCustomization::GetEndpoin
 	}
 	else
 	{
-		TArray<void*> RawData;
-		PropertyHandle->AccessRawData(RawData);
-		if (RawData.Num() > 1)
+		if (PropertyRawData.Num() > 1)
 		{
 			return FAppStyle::GetBrush("Sequencer.MultipleEvents");
 		}
@@ -1266,9 +1261,7 @@ FText FMovieSceneDirectorBlueprintEndpointCustomization::GetEndpointName() const
 	}
 	else
 	{
-		TArray<void*> RawData;
-		PropertyHandle->AccessRawData(RawData);
-		if (RawData.Num() != 1)
+		if (PropertyRawData.Num() != 1)
 		{
 			return LOCTEXT("MultipleValuesText", "Multiple Values");
 		}
@@ -1282,12 +1275,9 @@ void FMovieSceneDirectorBlueprintEndpointCustomization::GatherSequenceData(FSequ
 	TArray<UObject*> EditObjects;
 	GetEditObjects(EditObjects);
 
-	TArray<void*> RawData;
-	PropertyHandle->AccessRawData(RawData);
-
-	check(EditObjects.Num() == RawData.Num());
+	check(EditObjects.Num() == PropertyRawData.Num());
 
-	for (int32 Index = 0; Index < RawData.Num(); ++Index)
+	for (int32 Index = 0; Index < PropertyRawData.Num(); ++Index)
 	{
 		UMovieSceneSequence*       Sequence           = EditObjects[Index]->GetTypedOuter<UMovieSceneSequence>();
 		FMovieSceneSequenceEditor* SequenceEditor     = FMovieSceneSequenceEditor::Find(Sequence);
@@ -1297,7 +1287,7 @@ void FMovieSceneDirectorBlueprintEndpointCustomization::GatherSequenceData(FSequ
 		ensure(SequenceData.Blueprint == nullptr || SequenceData.Blueprint == SequenceDirectorBP);
 		SequenceData.Blueprint = SequenceDirectorBP;
 		SequenceData.EditObjects.Add(EditObjects[Index]);
-		SequenceData.RawData.Add(RawData[Index]);
+		SequenceData.RawData.Add(PropertyRawData[Index]);
 	}
 }
 
@@ -1438,7 +1428,11 @@ void FMovieSceneDirectorBlueprintEndpointCustomization::SetEndpoint(const FMovie
 	}
 
 	// Ensure that anything listening for property changed notifications are notified of the new binding
-	PropertyHandle->NotifyFinishedChangingProperties();
+	if (PropertyHandle.IsValid())
+	{
+		// Ensure that anything listening for property changed notifications are notified of the new binding
+		PropertyHandle->NotifyFinishedChangingProperties();
+	}
 
 	// Compile the blueprint now that clients have had a chance to update underlying data (we do this after to ensure we are compiling the correct data)
 	if (Blueprint)
@@ -1447,7 +1441,10 @@ void FMovieSceneDirectorBlueprintEndpointCustomization::SetEndpoint(const FMovie
 	}
 
 	// Forcibly update the panel now that our endpoint has changed
-	PropertyUtilities->ForceRefresh();
+	if (PropertyUtilities.IsValid())
+	{
+		PropertyUtilities->ForceRefresh();
+	}
 
 	if (NewEndpoint)
 	{
@@ -1467,18 +1464,15 @@ void FMovieSceneDirectorBlueprintEndpointCustomization::CreateEndpoint()
 
 	// Populate all the sequences represented by this customization
 	{
-		TArray<void*> RawData;
-		GetPropertyHandle()->AccessRawData(RawData);
-
 		TArray<UObject*> EditObjects;
 		GetEditObjects(EditObjects);
 
-		check(RawData.Num() == EditObjects.Num());
+		check(PropertyRawData.Num() == EditObjects.Num());
 
-		for (int32 Index = 0; Index < RawData.Num(); ++Index)
+		for (int32 Index = 0; Index < PropertyRawData.Num(); ++Index)
 		{
 			FSequenceData& SequenceData = PerSequenceData.FindOrAdd(EditObjects[Index]->GetTypedOuter<UMovieSceneSequence>());
-			SequenceData.RawData.Add(RawData[Index]);
+			SequenceData.RawData.Add(PropertyRawData[Index]);
 			SequenceData.EditObjects.Add(EditObjects[Index]);
 		}
 	}
@@ -1521,8 +1515,11 @@ void FMovieSceneDirectorBlueprintEndpointCustomization::CreateEndpoint()
 		LastNewEndpoint = NewEndpoint;
 	}
 
-	// Ensure that anything listening for property changed notifications are notified of the new binding
-	PropertyHandle->NotifyFinishedChangingProperties();
+	if (PropertyHandle.IsValid())
+	{
+		// Ensure that anything listening for property changed notifications are notified of the new binding
+		PropertyHandle->NotifyFinishedChangingProperties();
+	}
 
 	// Compile the blueprint now that clients have had a chance to update underlying data (we do this after to ensure we are compiling the correct data)
 	for (UBlueprint* Blueprint : BlueprintsToRecompile)
@@ -1530,7 +1527,10 @@ void FMovieSceneDirectorBlueprintEndpointCustomization::CreateEndpoint()
 		FKismetEditorUtilities::CompileBlueprint(Blueprint);
 	}
 
-	PropertyUtilities->ForceRefresh();
+	if (PropertyUtilities.IsValid())
+	{
+		PropertyUtilities->ForceRefresh();
+	}
 
 	// Focus the first created endpoint
 	if (LastNewEndpoint)
diff --git a/Engine/Source/Editor/MovieSceneTools/Private/MovieSceneDynamicBindingCustomization.cpp b/Engine/Source/Editor/MovieSceneTools/Private/MovieSceneDynamicBindingCustomization.cpp
--- a/Engine/Source/Editor/MovieSceneTools/Private/MovieSceneDynamicBindingCustomization.cpp
+++ b/Engine/Source/Editor/MovieSceneTools/Private/MovieSceneDynamicBindingCustomization.cpp
@@ -33,11 +33,12 @@ TSharedRef<IPropertyTypeCustomization> FMovieSceneDynamicBindingCustomization::M
 	return Instance;
 }
 
-TSharedRef<IPropertyTypeCustomization> FMovieSceneDynamicBindingCustomization::MakeInstance(UMovieScene* InMovieScene, FGuid InObjectBinding)
+TSharedRef<IPropertyTypeCustomization> FMovieSceneDynamicBindingCustomization::MakeInstance(UMovieScene* InMovieScene, FGuid InObjectBinding, int32 InBindingIndex)
 {
 	TSharedRef<FMovieSceneDynamicBindingCustomization> Instance = MakeShared<FMovieSceneDynamicBindingCustomization>();
 	Instance->EditedMovieScene = InMovieScene;
 	Instance->ObjectBinding = InObjectBinding;
+	Instance->BindingIndex = InBindingIndex;
 	return Instance;
 }
 
diff --git a/Engine/Source/Editor/MovieSceneTools/Public/MovieSceneDirectorBlueprintEndpointCustomization.h b/Engine/Source/Editor/MovieSceneTools/Public/MovieSceneDirectorBlueprintEndpointCustomization.h
--- a/Engine/Source/Editor/MovieSceneTools/Public/MovieSceneDirectorBlueprintEndpointCustomization.h
+++ b/Engine/Source/Editor/MovieSceneTools/Public/MovieSceneDirectorBlueprintEndpointCustomization.h
@@ -72,6 +72,9 @@ public:
 	/* Set the property handle to use for customization. Used when using the customization class as a helper. */
 	void SetPropertyHandle(TSharedPtr<IPropertyHandle> InPropertyHandle) { PropertyHandle = InPropertyHandle;}
 
+	/* Set the raw data to use for customization. Used when the customization class is used as a helper without a details view. */
+	void SetRawData(const TArray<void*>& InRawData) { PropertyRawData = InRawData; }
+
 protected:
 
 	using FPayloadVariableMap = TMap<FName, FMovieSceneDirectorBlueprintVariableValue, TInlineSetAllocator<8>>;
@@ -286,6 +289,9 @@ protected:
 	/** Property utilities for the property we're editing */
 	TSharedPtr<IPropertyUtilities> PropertyUtilities;
 
+	// RawData taken either from the PropertyHandle, or manually passed into the customization for cases this is used as a helper
+	TArray<void*> PropertyRawData;
+
 private:
 
 	void CollectQuickBindActions(FGraphActionListBuilderBase& OutAllActions, UBlueprint* Blueprint, FMovieSceneDirectorBlueprintEndpointDefinition EndpointDefinition);
diff --git a/Engine/Source/Editor/MovieSceneTools/Public/MovieSceneDynamicBindingCustomization.h b/Engine/Source/Editor/MovieSceneTools/Public/MovieSceneDynamicBindingCustomization.h
--- a/Engine/Source/Editor/MovieSceneTools/Public/MovieSceneDynamicBindingCustomization.h
+++ b/Engine/Source/Editor/MovieSceneTools/Public/MovieSceneDynamicBindingCustomization.h
@@ -26,7 +26,7 @@ class MOVIESCENETOOLS_API FMovieSceneDynamicBindingCustomization : public FMovie
 public:
 
 	static TSharedRef<IPropertyTypeCustomization> MakeInstance();
-	static TSharedRef<IPropertyTypeCustomization> MakeInstance(UMovieScene* InMovieScene, FGuid InObjectBinding);
+	static TSharedRef<IPropertyTypeCustomization> MakeInstance(UMovieScene* InMovieScene, FGuid InObjectBinding, int32 InBindingIndex=0);
 
 	virtual void CustomizeChildren(TSharedRef<IPropertyHandle> InPropertyHandle, IDetailChildrenBuilder& ChildBuilder, IPropertyTypeCustomizationUtils& CustomizationUtils) override;
 
diff --git a/Engine/Source/Editor/Sequencer/Private/MVVM/ViewModels/PossessableModel.cpp b/Engine/Source/Editor/Sequencer/Private/MVVM/ViewModels/PossessableModel.cpp
--- a/Engine/Source/Editor/Sequencer/Private/MVVM/ViewModels/PossessableModel.cpp
+++ b/Engine/Source/Editor/Sequencer/Private/MVVM/ViewModels/PossessableModel.cpp
@@ -139,7 +139,7 @@ const FSlateBrush* FPossessableModel::GetIconOverlayBrush() const
 		{
 			for (const FMovieSceneBindingReference& BindingReference : BindingReferences->GetReferences(ObjectBindingID))
 			{
-				if (const FSlateBrush* CustomBrush = BindingReference.CustomBinding ? BindingReference.CustomBinding->GetBindingTrackCustomIconOverlay() : nullptr)
+				if (const FSlateBrush* CustomBrush = BindingReference.CustomBinding ? BindingReference.CustomBinding->GetBindingTrackCustomIconOverlay().GetIcon() : nullptr)
 				{
 					return CustomBrush;
 				}
diff --git a/Engine/Source/Editor/Sequencer/Private/Sequencer.cpp b/Engine/Source/Editor/Sequencer/Private/Sequencer.cpp
--- a/Engine/Source/Editor/Sequencer/Private/Sequencer.cpp
+++ b/Engine/Source/Editor/Sequencer/Private/Sequencer.cpp
@@ -1602,7 +1602,7 @@ void FSequencer::UpdateSequencerCustomizations(const UMovieSceneSequence* Previo
 	// TODO: Allow customization of which custom bindings are allowed? For now just iterate over all subclasses and sort by priority
 
 	// Cache custom spawnable types
-	MovieSceneHelpers::GetPrioritySortedCustomBindingTypes(SupportedCustomBindingTypes);
+	RefreshSupportedCustomBindingTypes();
 }
 
 void FSequencer::RerunConstructionScripts()
@@ -5250,6 +5250,10 @@ TRange<FFrameNumber> FSequencer::GetRootTimeBounds() const
 	return RootTimeBounds;
 }
 
+void FSequencer::RefreshSupportedCustomBindingTypes()
+{
+	MovieSceneHelpers::GetPrioritySortedCustomBindingTypes(SupportedCustomBindingTypes);
+}
 
 void FSequencer::SetViewRange(TRange<double> NewViewRange, EViewRangeInterpolation Interpolation)
 {
diff --git a/Engine/Source/Editor/Sequencer/Private/Sequencer.h b/Engine/Source/Editor/Sequencer/Private/Sequencer.h
--- a/Engine/Source/Editor/Sequencer/Private/Sequencer.h
+++ b/Engine/Source/Editor/Sequencer/Private/Sequencer.h
@@ -853,6 +853,8 @@ public:
 	TRange<FFrameNumber> GetRootTimeBounds() const;
 
 	virtual TArrayView<const TSubclassOf<UMovieSceneCustomBinding>> GetSupportedCustomBindingTypes() const override { return SupportedCustomBindingTypes; }
+	
+	virtual void RefreshSupportedCustomBindingTypes() override;
 
 protected:
 
diff --git a/Engine/Source/Editor/Sequencer/Private/SequencerUtilities.cpp b/Engine/Source/Editor/Sequencer/Private/SequencerUtilities.cpp
--- a/Engine/Source/Editor/Sequencer/Private/SequencerUtilities.cpp
+++ b/Engine/Source/Editor/Sequencer/Private/SequencerUtilities.cpp
@@ -4429,164 +4429,4 @@ bool FSequencerUtilities::CanConvertToCustomBinding(TSharedRef<ISequencer> Seque
 	return false;
 }
 
-
-void FSequencerUtilities::AddConvertBindingMenu(FMenuBuilder& MenuBuilder, TSharedRef<ISequencer> Sequencer, const TArray<FSequencerChangeBindingInfo>& BindingsToConvert, TFunction<void()> OnBindingChanged)
-{
-	UMovieSceneSequence* Sequence = Sequencer->GetFocusedMovieSceneSequence();
-	UMovieScene* MovieScene = Sequence ? Sequence->GetMovieScene() : nullptr;
-
-	if (!MovieScene || BindingsToConvert.IsEmpty())
-	{
-		return;
-	}
-
-	// Can convert to possessable
-	if (Algo::AllOf(BindingsToConvert, [Sequencer, &Sequence, &MovieScene](const FSequencerChangeBindingInfo& BindingInfo) {
-		return CanConvertToPossessable(Sequencer, BindingInfo.BindingID, BindingInfo.BindingIndex );
-	}))
-	{
-		const TWeakPtr<ISequencer> WeakSequencer = Sequencer;
-
-		MenuBuilder.AddMenuEntry(
-			LOCTEXT("ConvertToPossessable", "Possessable"),
-			LOCTEXT("ConvertToPossessableTooltip", "Convert selected binding(s) to a possessable"),
-			FSlateIcon(),
-			FUIAction(FExecuteAction::CreateLambda([WeakSequencer, BindingsToConvert, OnBindingChanged]()
-			{
-				const TSharedPtr<ISequencer> Sequencer = WeakSequencer.Pin();
-				if (!Sequencer)
-				{
-					return;
-				}
-
-				const TSharedRef<ISequencer> SequencerRef = Sequencer.ToSharedRef();
-				ConvertBindings(SequencerRef, BindingsToConvert, [&SequencerRef](FGuid BindingID, int32 BindingIndex)
-					{
-						return FSequencerUtilities::ConvertToPossessable(SequencerRef, BindingID, BindingIndex);
-					},
-					OnBindingChanged);
-			})));
-	}
-	
-	// Can convert to custom bindings
-
-	TArrayView<const TSubclassOf<UMovieSceneCustomBinding>> PrioritySortedCustomBindingTypes = Sequencer->GetSupportedCustomBindingTypes();
-	for (const TSubclassOf<UMovieSceneCustomBinding>& CustomBindingType : PrioritySortedCustomBindingTypes)
-	{
-		if (Algo::AllOf(BindingsToConvert, [&Sequence, &MovieScene, &CustomBindingType, Sequencer](const FSequencerChangeBindingInfo& BindingInfo) 
-		{
-			return CanConvertToCustomBinding(Sequencer, BindingInfo.BindingID, CustomBindingType, BindingInfo.BindingIndex);
-		}))
-		{
-			const TWeakPtr<ISequencer> WeakSequencer = Sequencer;
-
-			MenuBuilder.AddMenuEntry(
-				CustomBindingType->GetDefaultObject<UMovieSceneCustomBinding>()->GetBindingTypePrettyName(),
-				FText::Format(LOCTEXT("ConvertToCustomBindingTooltip", "Convert selected binding to {0}"), CustomBindingType->GetDefaultObject<UMovieSceneCustomBinding>()->GetBindingTypePrettyName()),
-				FSlateIcon(),
-				FUIAction(FExecuteAction::CreateLambda([WeakSequencer, CustomBindingType, BindingsToConvert, OnBindingChanged]()
-				{
-					const TSharedPtr<ISequencer> Sequencer = WeakSequencer.Pin();
-					if (!Sequencer)
-					{
-						return;
-					}
-
-					const TSharedRef<ISequencer> SequencerRef = Sequencer.ToSharedRef();
-					ConvertBindings(SequencerRef, BindingsToConvert, [&SequencerRef, CustomBindingType](FGuid BindingID, int32 BindingIndex)
-						{
-							return FSequencerUtilities::ConvertToCustomBinding(SequencerRef, BindingID, CustomBindingType, BindingIndex);
-						},
-						OnBindingChanged);
-				})));
-		}
-	}
-}
-
-void FSequencerUtilities::ConvertBindings(const TSharedRef<ISequencer>& InSequencer
-	, const TArray<FSequencerChangeBindingInfo>& InBindingsToConvert
-	, TFunction<FMovieScenePossessable*(FGuid, int32)> InDoConvert
-	, TFunction<void()> InOnBindingChanged)
-{
-	using namespace UE::Sequencer;
-
-	UMovieSceneSequence* const Sequence = InSequencer->GetFocusedMovieSceneSequence();
-	if (!IsValid(Sequence))
-	{
-		return;
-	}
-
-	UMovieScene* const MovieScene = Sequence->GetMovieScene();
-	if (!IsValid(MovieScene))
-	{
-		return;
-	}
-
-	if (InBindingsToConvert.Num() == 0)
-	{
-		return;
-	}
-
-	if (MovieScene->IsReadOnly())
-	{
-		FSequencerUtilities::ShowReadOnlyError();
-		return;
-	}
-
-	const FScopedTransaction Transaction(LOCTEXT("ConvertSelectedNodesPossessable", "Convert Selected Nodes to Possessables"));
-	MovieScene->Modify();
-
-	FScopedSlowTask SlowTask(InBindingsToConvert.Num(), LOCTEXT("ConvertPossessablesProgress", "Converting Selected Spawnable Nodes to Possessables"));
-	SlowTask.MakeDialog(true);
-
-	TArray<AActor*> PossessedActors;
-	for (const FSequencerChangeBindingInfo& BindingInfo : InBindingsToConvert)
-	{
-		SlowTask.EnterProgressFrame();
-
-		if (FMovieScenePossessable* Possessable = InDoConvert(BindingInfo.BindingID, BindingInfo.BindingIndex))
-		{
-			InSequencer->ForceEvaluate();
-
-			for (TWeakObjectPtr<> WeakObject : InSequencer->FindBoundObjects(Possessable->GetGuid(), InSequencer->GetFocusedTemplateID()))
-			{
-				if (AActor* PossessedActor = Cast<AActor>(WeakObject.Get()))
-				{
-					PossessedActors.Add(PossessedActor);
-				}
-			}
-
-			if (GWarn->ReceivedUserCancel())
-			{
-				break;
-			}
-		}
-	}
-
-	if (PossessedActors.Num())
-	{
-		const bool bNotifySelectionChanged = true;
-		const bool bDeselectBSP = true;
-		const bool bWarnAboutTooManyActors = false;
-		const bool bSelectEvenIfHidden = false;
-
-		GEditor->GetSelectedActors()->Modify();
-		GEditor->GetSelectedActors()->BeginBatchSelectOperation();
-		GEditor->SelectNone(bNotifySelectionChanged, bDeselectBSP, bWarnAboutTooManyActors);
-		for (auto PossessedActor : PossessedActors)
-		{
-			GEditor->SelectActor(PossessedActor, true, bNotifySelectionChanged, bSelectEvenIfHidden);
-		}
-		GEditor->GetSelectedActors()->EndBatchSelectOperation();
-		GEditor->NoteSelectionChange();
-
-		InSequencer->NotifyMovieSceneDataChanged(EMovieSceneDataChangeType::MovieSceneStructureItemsChanged);
-	}
-
-	if (InOnBindingChanged != nullptr)
-	{
-		InOnBindingChanged();
-	}
-}
-
 #undef LOCTEXT_NAMESPACE
diff --git a/Engine/Source/Editor/Sequencer/Public/ISequencer.h b/Engine/Source/Editor/Sequencer/Public/ISequencer.h
--- a/Engine/Source/Editor/Sequencer/Public/ISequencer.h
+++ b/Engine/Source/Editor/Sequencer/Public/ISequencer.h
@@ -881,6 +881,8 @@ public:
 	*/
 	virtual TArrayView<const TSubclassOf<UMovieSceneCustomBinding>> GetSupportedCustomBindingTypes() const { static TArray<TSubclassOf<UMovieSceneCustomBinding>> EmptyArray; return EmptyArray; }
 
+	virtual void RefreshSupportedCustomBindingTypes() {}
+
 	/** @return True if the Sequencer is currently limiting viewport selection to only Sequencer objects. */
 	virtual bool IsViewportSelectionLimited() const = 0;
 
diff --git a/Engine/Source/Editor/Sequencer/Public/SequencerUtilities.h b/Engine/Source/Editor/Sequencer/Public/SequencerUtilities.h
--- a/Engine/Source/Editor/Sequencer/Public/SequencerUtilities.h
+++ b/Engine/Source/Editor/Sequencer/Public/SequencerUtilities.h
@@ -304,10 +304,5 @@ struct SEQUENCER_API FSequencerUtilities
 	// Functions allowing menus to be built for modifying bindings
 	static void AddChangeClassMenu(FMenuBuilder& MenuBuilder, TSharedRef<ISequencer> Sequencer, const TArray<FSequencerChangeBindingInfo>& BindingsToConvert, TFunction<void()> OnBindingChanged);
 	static void HandleTemplateActorClassPicked(UClass* ChosenClass, TSharedRef<ISequencer> Sequencer, const TArray<FSequencerChangeBindingInfo>& BindingsToConvert, TFunction<void()> OnBindingChanged);
-	static void AddConvertBindingMenu(FMenuBuilder& MenuBuilder, TSharedRef<ISequencer> Sequencer, const TArray<FSequencerChangeBindingInfo>& BindingsToConvert, TFunction<void()> OnBindingChanged);
 
-	static void ConvertBindings(const TSharedRef<ISequencer>& InSequencer
-		, const TArray<FSequencerChangeBindingInfo>& InBindingsToConvert
-		, TFunction<FMovieScenePossessable*(FGuid, int32)> InDoConvert
-		, TFunction<void()> InOnBindingChanged);
 };
diff --git a/Engine/Source/Editor/UMGEditor/Private/WidgetBlueprintEditor.cpp b/Engine/Source/Editor/UMGEditor/Private/WidgetBlueprintEditor.cpp
--- a/Engine/Source/Editor/UMGEditor/Private/WidgetBlueprintEditor.cpp
+++ b/Engine/Source/Editor/UMGEditor/Private/WidgetBlueprintEditor.cpp
@@ -2378,7 +2378,7 @@ void FWidgetBlueprintEditor::AddDynamicPossessionMenu(FMenuBuilder& MenuBuilder,
 		// Register details customizations for this instance
 		StructureDetailsView->GetDetailsView()->RegisterInstancedCustomPropertyTypeLayout(
 			FMovieSceneDynamicBinding::StaticStruct()->GetFName(),
-			FOnGetPropertyTypeCustomizationInstance::CreateStatic(&FMovieSceneDynamicBindingCustomization::MakeInstance, MovieScene, ObjectId));
+			FOnGetPropertyTypeCustomizationInstance::CreateStatic(&FMovieSceneDynamicBindingCustomization::MakeInstance, MovieScene, ObjectId, 0));
 
 		// We can't just show the FMovieSceneDynamicBinding struct in the details view, because Slate only uses
 		// the above details view customization for *properties* (not for the root object). So here we put a copy of
diff --git a/Engine/Source/Runtime/MovieScene/Private/Bindings/MovieSceneCustomBinding.cpp b/Engine/Source/Runtime/MovieScene/Private/Bindings/MovieSceneCustomBinding.cpp
new file mode 100644
--- /dev/null
+++ b/Engine/Source/Runtime/MovieScene/Private/Bindings/MovieSceneCustomBinding.cpp
@@ -0,0 +1,6 @@
+// Copyright Epic Games, Inc. All Rights Reserved.
+
+#include "Bindings/MovieSceneCustomBinding.h"
+
+const int32 UMovieSceneCustomBinding::BaseEnginePriority = 0;
+const int32 UMovieSceneCustomBinding::BaseCustomPriority = 100;
\ No newline at end of file
diff --git a/Engine/Source/Runtime/MovieScene/Private/Bindings/MovieSceneReplaceableBinding.cpp b/Engine/Source/Runtime/MovieScene/Private/Bindings/MovieSceneReplaceableBinding.cpp
--- a/Engine/Source/Runtime/MovieScene/Private/Bindings/MovieSceneReplaceableBinding.cpp
+++ b/Engine/Source/Runtime/MovieScene/Private/Bindings/MovieSceneReplaceableBinding.cpp
@@ -40,9 +40,9 @@ void UMovieSceneReplaceableBindingBase::SetupDefaults(UObject* SpawnedObject, FG
 	}
 }
 
-const FSlateBrush* UMovieSceneReplaceableBindingBase::GetBindingTrackCustomIconOverlay() const
+FSlateIcon UMovieSceneReplaceableBindingBase::GetBindingTrackCustomIconOverlay() const
 {
-	return FAppStyle::GetBrush("Sequencer.ReplaceableIconOverlay");
+	return FSlateIcon(FAppStyle::GetAppStyleSetName(), "Sequencer.ReplaceableIconOverlay");
 }
 
 FText UMovieSceneReplaceableBindingBase::GetBindingTrackIconTooltip() const
diff --git a/Engine/Source/Runtime/MovieScene/Private/Bindings/MovieSceneSpawnableBinding.cpp b/Engine/Source/Runtime/MovieScene/Private/Bindings/MovieSceneSpawnableBinding.cpp
--- a/Engine/Source/Runtime/MovieScene/Private/Bindings/MovieSceneSpawnableBinding.cpp
+++ b/Engine/Source/Runtime/MovieScene/Private/Bindings/MovieSceneSpawnableBinding.cpp
@@ -202,9 +202,9 @@ void UMovieSceneSpawnableBindingBase::SetupDefaults(UObject* SpawnedObject, FGui
 	}
 }
 
-const FSlateBrush* UMovieSceneSpawnableBindingBase::GetBindingTrackCustomIconOverlay() const
+FSlateIcon UMovieSceneSpawnableBindingBase::GetBindingTrackCustomIconOverlay() const
 {
-	return FAppStyle::GetBrush("Sequencer.SpawnableIconOverlay");
+	return FSlateIcon(FAppStyle::GetAppStyleSetName(), "Sequencer.SpawnableIconOverlay");
 }
 
 FText UMovieSceneSpawnableBindingBase::GetBindingTrackIconTooltip() const
diff --git a/Engine/Source/Runtime/MovieScene/Private/MovieSceneCommonHelpers.cpp b/Engine/Source/Runtime/MovieScene/Private/MovieSceneCommonHelpers.cpp
--- a/Engine/Source/Runtime/MovieScene/Private/MovieSceneCommonHelpers.cpp
+++ b/Engine/Source/Runtime/MovieScene/Private/MovieSceneCommonHelpers.cpp
@@ -993,6 +993,11 @@ void MovieSceneHelpers::GetPrioritySortedCustomBindingTypes(TArray<const TSubcla
 				continue;
 			}
 #endif
+			// Skip SKEL and REINST classes.
+			if (It->GetName().StartsWith(TEXT("SKEL_")) || It->GetName().StartsWith(TEXT("REINST_")))
+			{
+				continue;
+			}
 			OutCustomBindingTypes.Add(*It);
 		}
 	}// Sort by spawner priority to allow disambiguation for similar object types
diff --git a/Engine/Source/Runtime/MovieScene/Public/Bindings/MovieSceneCustomBinding.h b/Engine/Source/Runtime/MovieScene/Public/Bindings/MovieSceneCustomBinding.h
--- a/Engine/Source/Runtime/MovieScene/Public/Bindings/MovieSceneCustomBinding.h
+++ b/Engine/Source/Runtime/MovieScene/Public/Bindings/MovieSceneCustomBinding.h
@@ -3,6 +3,8 @@
 #pragma once
 
 #include "MovieSceneSequenceID.h"
+#include "MovieSceneBindingProxy.h"
+#include "Textures/SlateIcon.h"
 #include "MovieSceneCustomBinding.generated.h"
 
 namespace UE
@@ -31,10 +33,27 @@ struct MOVIESCENE_API FMovieSceneBindingResolveResult
 	TObjectPtr<UObject> Object = nullptr;
 };
 
+/*
+* Blueprint-specific resolution context for custom bindings.
+*/
+USTRUCT(BlueprintType)
+struct MOVIESCENE_API FMovieSceneBindingResolveContext
+{
+	GENERATED_BODY()
+
+	/* The world context*/
+	UPROPERTY(EditAnywhere, BlueprintReadWrite, Category="Default")
+	TObjectPtr<UObject> WorldContext;
+	
+	/* Binding for the bound object currently evaluating this condition if applicable (BindingId will be invalid for conditions on global tracks/sections). */
+	UPROPERTY(EditAnywhere, BlueprintReadWrite, Category="Default")
+	FMovieSceneBindingProxy Binding;
+};
+
 /**
  * A custom binding. Allows users to define their own binding resolution types, including dynamic 'Replaceable' bindings with previews in editor, as well as Spawnable types.
  */
-UCLASS(abstract, BlueprintType, DefaultToInstanced, EditInlineNew)
+UCLASS(abstract, DefaultToInstanced, EditInlineNew)
 class MOVIESCENE_API UMovieSceneCustomBinding
 	: public UObject
 {
@@ -42,6 +61,15 @@ public:
 
 	GENERATED_BODY()
 
+	static const int32 BaseEnginePriority;
+	static const int32 BaseCustomPriority;
+
+	UFUNCTION(BlueprintCallable, Category = "Sequencer|Binding")
+	static int32 GetBaseEnginePriority() { return BaseEnginePriority; }
+
+	UFUNCTION(BlueprintCallable, Category = "Sequencer|Binding")
+	static int32 GetBaseCustomPriority() { return BaseCustomPriority; }
+
 	/* Must be implemented.
 	* Resolve the custom binding based on the passed in context. May return an existing UObject or spawn a new one. 
 	*/
@@ -55,7 +83,7 @@ public:
 	  * @return Custom binding priority in order to sort the list of custom binding types.
 	  * If several custom binding types support the creation of bindings from the same object types, the one with the highest priority will be picked.
 	  */
-	virtual int32 GetCustomBindingPriority() const { return 0; }
+	virtual int32 GetCustomBindingPriority() const { return BaseEnginePriority; }
 
 	/*
 	* Must be implemented. Called by Sequencer to determine whether this custom binding type supports binding the given object.
@@ -105,7 +133,7 @@ public:
 	/*
 	* Allows the custom binding to optionally provide a custom icon overlay for the object binding track.
 	*/
-	virtual const FSlateBrush* GetBindingTrackCustomIconOverlay() const { return nullptr; }
+	virtual FSlateIcon GetBindingTrackCustomIconOverlay() const { return FSlateIcon(); }
 
 	/*
 	* Allows the custom binding to optionally provide a custom tooltip to show when hovering over the icon area in the object binding track.
diff --git a/Engine/Source/Runtime/MovieScene/Public/Bindings/MovieSceneReplaceableBinding.h b/Engine/Source/Runtime/MovieScene/Public/Bindings/MovieSceneReplaceableBinding.h
--- a/Engine/Source/Runtime/MovieScene/Public/Bindings/MovieSceneReplaceableBinding.h
+++ b/Engine/Source/Runtime/MovieScene/Public/Bindings/MovieSceneReplaceableBinding.h
@@ -27,7 +27,7 @@ class UMovieSceneSpawnableBindingBase;
  * UMovieSceneReplaceableActorBinding as an example is the simplest type of replaceable binding and provides no method for binding at runtime and relies on the LevelSequenceActor's Binding Override 
  * mechanism to bind an actor at runtime.
  */
-UCLASS(Abstract, BlueprintType)
+UCLASS(Abstract)
 class MOVIESCENE_API UMovieSceneReplaceableBindingBase
 	: public UMovieSceneCustomBinding
 {
@@ -42,7 +42,7 @@ public:
 	virtual UClass* GetBoundObjectClass() const override;
 #if WITH_EDITOR
 	virtual	void SetupDefaults(UObject* SpawnedObject, FGuid ObjectBindingId, UMovieScene& OwnerMovieScene) override;
-	virtual const FSlateBrush* GetBindingTrackCustomIconOverlay() const override;
+	virtual FSlateIcon GetBindingTrackCustomIconOverlay() const override;
 	virtual FText GetBindingTrackIconTooltip() const override;
 
 	virtual bool SupportsConversionFromBinding(const FMovieSceneBindingReference& BindingReference, const UObject* SourceObject) const override;
diff --git a/Engine/Source/Runtime/MovieScene/Public/Bindings/MovieSceneSpawnableBinding.h b/Engine/Source/Runtime/MovieScene/Public/Bindings/MovieSceneSpawnableBinding.h
--- a/Engine/Source/Runtime/MovieScene/Public/Bindings/MovieSceneSpawnableBinding.h
+++ b/Engine/Source/Runtime/MovieScene/Public/Bindings/MovieSceneSpawnableBinding.h
@@ -25,7 +25,7 @@ struct FSlateBrush;
  * In doing so, they could choose to just override GetSpawnObjectClass, PostSpawnObject, and PreDestroyObject for example to do custom post-spawn setup on a character mesh,
  * or they could choose to fully override SpawnObject and DestroySpawnedObject and do their own custom logic for spawning completely.
  */
-UCLASS(abstract, BlueprintType)
+UCLASS(abstract)
 class MOVIESCENE_API UMovieSceneSpawnableBindingBase
 	: public UMovieSceneCustomBinding
 {
@@ -79,7 +79,7 @@ public:
 #if WITH_EDITOR
 	/* UMovieSceneCustomBinding overrides */
 	virtual	void SetupDefaults(UObject* SpawnedObject, FGuid ObjectBindingId, UMovieScene& OwnerMovieScene) override;
-	virtual const FSlateBrush* GetBindingTrackCustomIconOverlay() const override;
+	virtual FSlateIcon GetBindingTrackCustomIconOverlay() const override;
 	virtual FText GetBindingTrackIconTooltip() const override;
 #endif
 protected:
diff --git a/Engine/Source/Runtime/MovieSceneTracks/Private/Bindings/MovieSceneReplaceableActorBinding.cpp b/Engine/Source/Runtime/MovieSceneTracks/Private/Bindings/MovieSceneReplaceableActorBinding.cpp
--- a/Engine/Source/Runtime/MovieSceneTracks/Private/Bindings/MovieSceneReplaceableActorBinding.cpp
+++ b/Engine/Source/Runtime/MovieSceneTracks/Private/Bindings/MovieSceneReplaceableActorBinding.cpp
@@ -20,11 +20,83 @@ FText UMovieSceneReplaceableActorBinding::GetBindingTypePrettyName() const
 {
 	return LOCTEXT("MovieSceneReplaceableActorBinding", "Replaceable Actor");
 }
-#endif
 
 TSubclassOf<UMovieSceneSpawnableBindingBase> UMovieSceneReplaceableActorBinding::GetInnerSpawnableClass() const
 {
 	return UMovieSceneSpawnableActorBinding::StaticClass();
 }
 
+
+FText UMovieSceneReplaceableActorBinding_BPBase::GetBindingTypePrettyName() const
+{
+	return BindingTypePrettyName.IsEmpty() ? GetClass()->GetDisplayNameText() : BindingTypePrettyName;
+}
+
+FText UMovieSceneReplaceableActorBinding_BPBase::GetBindingTrackIconTooltip() const
+{
+	return BindingTypeTooltip.IsEmpty() ? Super::GetBindingTrackIconTooltip() : BindingTypeTooltip;
+}
+
+
+void UMovieSceneReplaceableActorBinding_BPBase::OnBindingAddedOrChanged(UMovieScene& OwnerMovieScene)
+{
+	if ((PreviewSpawnableType == nullptr && PreviewSpawnable)
+		|| (PreviewSpawnableType != nullptr && (!PreviewSpawnable || PreviewSpawnable->GetClass() != PreviewSpawnableType)))
+	{
+		if (PreviewSpawnableType == nullptr)
+		{
+			PreviewSpawnable = nullptr;
+		}
+		else
+		{
+			PreviewSpawnable = NewObject<UMovieSceneSpawnableBindingBase>(&OwnerMovieScene, PreviewSpawnableType, NAME_None, RF_Transactional);
+		}
+	}
+}
+
+
+UMovieSceneCustomBinding* UMovieSceneReplaceableActorBinding_BPBase::CreateCustomBindingFromBinding(const FMovieSceneBindingReference& BindingReference, UObject* SourceObject, UMovieScene& OwnerMovieScene)
+{
+	// We override this specifically to initialize PreviewSpawnableType on conversions so that we end up with a custom preview from whatever object was passed in in that case.
+	UMovieSceneReplaceableActorBinding_BPBase* NewCustomBinding = nullptr;
+
+	const FName TemplateName = MakeUniqueObjectName(&OwnerMovieScene, UObject::StaticClass(), SourceObject ? SourceObject->GetFName() : GetClass()->GetFName());
+	const FName InstancedBindingName = MakeUniqueObjectName(&OwnerMovieScene, UObject::StaticClass(), *FString(TemplateName.ToString() + TEXT("_CustomBinding")));
+
+	NewCustomBinding = NewObject<UMovieSceneReplaceableActorBinding_BPBase>(&OwnerMovieScene, GetClass(), InstancedBindingName, RF_Transactional);
+
+	NewCustomBinding->PreviewSpawnable = NewCustomBinding->CreateInnerSpawnable(SourceObject, OwnerMovieScene);
+	NewCustomBinding->InitReplaceableBinding(SourceObject, OwnerMovieScene);
+	return NewCustomBinding;
+}
+#endif
+
+FMovieSceneBindingResolveResult UMovieSceneReplaceableActorBinding_BPBase::ResolveRuntimeBindingInternal(const FMovieSceneBindingResolveParams& ResolveParams, int32 BindingIndex, TSharedRef<const UE::MovieScene::FSharedPlaybackState> SharedPlaybackState) const
+{
+	FMovieSceneBindingResolveContext ResolveContext { SharedPlaybackState->GetPlaybackContext(), FMovieSceneBindingProxy(ResolveParams.ObjectBindingID, ResolveParams.Sequence)};
+	return BP_ResolveRuntimeBinding(ResolveContext);
+}
+
+void UMovieSceneReplaceableActorBinding_BPBase::InitReplaceableBinding(UObject* SourceObject, UMovieScene& OwnerMovieScene)
+{
+	BP_InitReplaceableBinding(SourceObject, &OwnerMovieScene);
+}
+
+bool UMovieSceneReplaceableActorBinding_BPBase::SupportsBindingCreationFromObject(const UObject* SourceObject) const
+{
+	return BP_SupportsBindingCreationFromObject(SourceObject);
+}
+
+void UMovieSceneReplaceableActorBinding_BPBase::BP_InitReplaceableBinding_Implementation(UObject* SourceObject, UMovieScene* OwnerMovieScene)
+{
+	// Default to nothing
+}
+
+bool UMovieSceneReplaceableActorBinding_BPBase::BP_SupportsBindingCreationFromObject_Implementation(const UObject* SourceObject) const
+{
+	// Defaultly we just ensure the object is an actor
+	return SourceObject && SourceObject->IsA<AActor>();
+}
+
+
 #undef LOCTEXT_NAMESPACE
\ No newline at end of file
diff --git a/Engine/Source/Runtime/MovieSceneTracks/Private/Bindings/MovieSceneReplaceableDirectorBlueprintBinding.cpp b/Engine/Source/Runtime/MovieSceneTracks/Private/Bindings/MovieSceneReplaceableDirectorBlueprintBinding.cpp
--- a/Engine/Source/Runtime/MovieSceneTracks/Private/Bindings/MovieSceneReplaceableDirectorBlueprintBinding.cpp
+++ b/Engine/Source/Runtime/MovieSceneTracks/Private/Bindings/MovieSceneReplaceableDirectorBlueprintBinding.cpp
@@ -69,8 +69,4 @@ FMovieSceneBindingResolveResult UMovieSceneReplaceableDirectorBlueprintBinding::
 	return ResolveResult;
 }
 
-UMovieSceneSpawnableBindingBase* UMovieSceneReplaceableDirectorBlueprintBinding::CreateInnerSpawnable(UObject* SourceObject, UMovieScene& OwnerMovieScene)
-{
-	return UMovieSceneReplaceableBindingBase::CreateInnerSpawnable(SourceObject, OwnerMovieScene);
-}
 #undef LOCTEXT_NAMESPACE
\ No newline at end of file
diff --git a/Engine/Source/Runtime/MovieSceneTracks/Public/Bindings/MovieSceneReplaceableActorBinding.h b/Engine/Source/Runtime/MovieSceneTracks/Public/Bindings/MovieSceneReplaceableActorBinding.h
--- a/Engine/Source/Runtime/MovieSceneTracks/Public/Bindings/MovieSceneReplaceableActorBinding.h
+++ b/Engine/Source/Runtime/MovieSceneTracks/Public/Bindings/MovieSceneReplaceableActorBinding.h
@@ -11,7 +11,7 @@
 * An implementation of UMovieSceneReplaceableBindingBase that uses UMovieSceneSpawnableActorBinding as the preview spawnable,
 * and has no implementation of ResolveRuntimeBindingInternal, relying instead of Sequencer's built in BindingOverride mechanism for binding at runtime.
 */
-UCLASS(BlueprintType, MinimalAPI, EditInlineNew, DefaultToInstanced, Meta=(DisplayName="Replaceable Actor"))
+UCLASS(MinimalAPI, BlueprintType, EditInlineNew, DefaultToInstanced, Meta=(DisplayName="Replaceable Actor"))
 class UMovieSceneReplaceableActorBinding
 	: public UMovieSceneReplaceableBindingBase
 {
@@ -41,7 +41,89 @@ protected:
 
 	MOVIESCENETRACKS_API virtual TSubclassOf<UMovieSceneSpawnableBindingBase> GetInnerSpawnableClass() const override;
 
-	virtual int32 GetCustomBindingPriority() const override { return 9; }
+	virtual int32 GetCustomBindingPriority() const override { return BaseEnginePriority + 1; }
 
 };
 
+/* Base class for Custom Replaceable Binding classes implemented by Blueprints */
+
+UCLASS(Abstract, Blueprintable, MinimalAPI, EditInlineNew, DefaultToInstanced, Meta=(ShowWorldContextPin, DisplayName="Replaceable Blueprint Base"))
+class UMovieSceneReplaceableActorBinding_BPBase
+	: public UMovieSceneReplaceableBindingBase
+{
+public:
+
+	GENERATED_BODY()
+
+public:
+
+	/* MovieSceneCustomBinding overrides*/
+#if WITH_EDITOR
+	FText GetBindingTypePrettyName() const override final;
+	virtual FText GetBindingTrackIconTooltip() const override final;
+
+	void OnBindingAddedOrChanged(UMovieScene& OwnerMovieScene) override final;
+#endif
+
+#if WITH_EDITORONLY_DATA
+	
+	/* Name to show in Sequencer for the custom binding type.*/
+	UPROPERTY(EditDefaultsOnly, BlueprintReadOnly, Category="Sequencer")
+	FText BindingTypePrettyName;
+
+	/* Tooltip to show in Sequencer for the custom binding type.*/
+	UPROPERTY(EditDefaultsOnly, BlueprintReadOnly, Category="Sequencer")
+	FText BindingTypeTooltip;
+#endif
+	
+	/*
+	Priority with which to consider this binding type over others when considering binding an object to Sequencer.
+	As a guideline, a priority of BaseEnginePriority will ensure that engine types(such as Spawnable Actor, Replaceable Actor) will
+	be higher priority than your custom binding, and so your binding type will not automatically be created(but may be converted to manually).
+	A priority of BaseCustomPriority and higher will ensure that your binding type is considered more highly than engine types,
+	so if your binding type's 'SupportsBindingCreationFromObject' returns true for an object, your binding type will be created by default
+	rather than an engine type.
+	*/
+	UPROPERTY(EditDefaultsOnly, BlueprintReadOnly, Category = "Sequencer")
+	int32 CustomBindingPriority = BaseEnginePriority;
+
+	// Preview Spawnable Type to use for this replaceable
+	UPROPERTY(EditDefaultsOnly, BlueprintReadOnly, Category="Preview", meta=(AllowAbstract=false))
+	TSubclassOf<UMovieSceneSpawnableBindingBase> PreviewSpawnableType;
+
+	int32 GetCustomBindingPriority() const override final { return CustomBindingPriority; }	/* Blueprint Interface */
+
+	/*
+	* Must be implemented. Called during non-editor/runtime to resolve the binding dynamically. In editor worlds/Sequencer will instead use the PreviewSpawnable binding to spawn a preview object.
+	* If no object is returned, Sequencer's BindingOverrides can still be used to dynamically bind the object.
+	*/
+	UFUNCTION(BlueprintImplementableEvent, CallInEditor, meta=(DisplayName = "Resolve Runtime Binding"))
+	FMovieSceneBindingResolveResult BP_ResolveRuntimeBinding(const FMovieSceneBindingResolveContext& ResolveContext) const;
+
+	/* Called after binding creation to allow the replaceable to initialize any data members from the source object. */
+	UFUNCTION(BlueprintNativeEvent, CallInEditor, meta=(DisplayName = "Init Replaceable Binding"))
+	void BP_InitReplaceableBinding(UObject* SourceObject, UMovieScene* OwnerMovieScene);
+
+	/* Called on the binding to determine whether this binding type supports creating a binding from the passed in object. */
+	UFUNCTION(BlueprintNativeEvent, CallInEditor, meta=(DisplayName = "Supports Binding Creation From Object"))
+	bool BP_SupportsBindingCreationFromObject(const UObject* SourceObject) const;
+
+protected:
+	/* MovieSceneReplaceableBindingBase overrides*/
+
+#if WITH_EDITOR
+	UMovieSceneCustomBinding* CreateCustomBindingFromBinding(const FMovieSceneBindingReference& BindingReference, UObject* SourceObject, UMovieScene& OwnerMovieScene) override final;
+#endif
+
+	// By default we return nullptr here, as we rely on Sequencer's BindingOverride mechanism to bind these actors during runtime.
+	// This can be overridden if desired in subclasses to provide a different way to resolve to an actor at runtime while still using spawnable actor as the preview.
+	FMovieSceneBindingResolveResult ResolveRuntimeBindingInternal(const FMovieSceneBindingResolveParams& ResolveParams, int32 BindingIndex, TSharedRef<const UE::MovieScene::FSharedPlaybackState> SharedPlaybackState) const override;
+	
+	void InitReplaceableBinding(UObject* SourceObject, UMovieScene& OwnerMovieScene) override final;
+
+	TSubclassOf<UMovieSceneSpawnableBindingBase> GetInnerSpawnableClass() const override final { return PreviewSpawnableType; }
+	
+	bool SupportsBindingCreationFromObject(const UObject* SourceObject) const override final;
+
+
+};
\ No newline at end of file
diff --git a/Engine/Source/Runtime/MovieSceneTracks/Public/Bindings/MovieSceneReplaceableDirectorBlueprintBinding.h b/Engine/Source/Runtime/MovieSceneTracks/Public/Bindings/MovieSceneReplaceableDirectorBlueprintBinding.h
--- a/Engine/Source/Runtime/MovieSceneTracks/Public/Bindings/MovieSceneReplaceableDirectorBlueprintBinding.h
+++ b/Engine/Source/Runtime/MovieSceneTracks/Public/Bindings/MovieSceneReplaceableDirectorBlueprintBinding.h
@@ -35,7 +35,6 @@ public:
 	FMovieSceneDynamicBinding DynamicBinding;
 
 	// Preview Spawnable Type to use for this replaceable
-	// TODO: Do I want to consider moving this sort of thing to the base class so any replaceable type can choose any other spawnable type?
 	UPROPERTY(EditAnywhere, BlueprintReadWrite, Category="Sequencer")
 	TSubclassOf<UMovieSceneSpawnableBindingBase> PreviewSpawnableType;
 
@@ -56,7 +55,5 @@ protected:
 	MOVIESCENETRACKS_API virtual TSubclassOf<UMovieSceneSpawnableBindingBase> GetInnerSpawnableClass() const override { return PreviewSpawnableType; }
 	
 	virtual bool SupportsBindingCreationFromObject(const UObject* SourceObject) const override { return true; }
-
-	virtual UMovieSceneSpawnableBindingBase* CreateInnerSpawnable(UObject* SourceObject, UMovieScene& OwnerMovieScene) override;
 };
 
diff --git a/Engine/Source/Runtime/MovieSceneTracks/Public/Bindings/MovieSceneSpawnableActorBinding.h b/Engine/Source/Runtime/MovieSceneTracks/Public/Bindings/MovieSceneSpawnableActorBinding.h
--- a/Engine/Source/Runtime/MovieSceneTracks/Public/Bindings/MovieSceneSpawnableActorBinding.h
+++ b/Engine/Source/Runtime/MovieSceneTracks/Public/Bindings/MovieSceneSpawnableActorBinding.h
@@ -13,7 +13,7 @@
  * Can be overridden in C++ or blueprint to provide an Actor class and to add custom PostSpawnObject behavior such as mesh setup based on an asset.
  * The below UMovieSceneSpawnableActorBinding class implements this base class and replicates the old FMovieSceneSpawnable behavior by using a specified Actor template to spawn an Actor and can be used out of the box.
  */
-UCLASS(Abstract, Blueprintable, MinimalAPI)
+UCLASS(Abstract, MinimalAPI)
 class UMovieSceneSpawnableActorBindingBase
 	: public UMovieSceneSpawnableBindingBase
 {
@@ -86,7 +86,7 @@ protected:
 	MOVIESCENETRACKS_API UClass* GetBoundObjectClass() const override { return GetActorClass(); }
 #if WITH_EDITOR
 	/* MovieSceneCustomBinding overrides*/
-	virtual int32 GetCustomBindingPriority() const override { return 10; }
+	virtual int32 GetCustomBindingPriority() const override { return BaseEnginePriority + 2; }
 #endif
 
 protected:
