Sequencer: Assorted UX and crash fixes for conditions and bindings.

diff --git a/Engine/Plugins/MovieScene/LevelSequenceEditor/Source/LevelSequenceEditor/Private/LevelSequenceEditorSubsystem.cpp b/Engine/Plugins/MovieScene/LevelSequenceEditor/Source/LevelSequenceEditor/Private/LevelSequenceEditorSubsystem.cpp
--- a/Engine/Plugins/MovieScene/LevelSequenceEditor/Source/LevelSequenceEditor/Private/LevelSequenceEditorSubsystem.cpp
+++ b/Engine/Plugins/MovieScene/LevelSequenceEditor/Source/LevelSequenceEditor/Private/LevelSequenceEditorSubsystem.cpp
@@ -613,13 +613,7 @@ void ULevelSequenceEditorSubsystem::Initialize(FSubsystemCollectionBase& Collect
 		}));
 
 	BindingPropertiesMenuExtender->AddMenuExtension("ConvertBinding", EExtensionHook::First, CommandList, FMenuExtensionDelegate::CreateLambda([this](FMenuBuilder& MenuBuilder) {
-		FFormatNamedArguments Args;
-
-		MenuBuilder.AddSubMenu(
-			FText::Format(LOCTEXT("BindingProperties", "Binding Properties"), Args),
-			FText::Format(LOCTEXT("BindingPropertiesTooltip", "Modify the actor and object bindings for this track"), Args),
-			FNewMenuDelegate::CreateLambda([this](FMenuBuilder& SubMenuBuilder) 
-			{ AddConvertBindingsMenu(SubMenuBuilder); }));
+		 AddConvertBindingsMenu(MenuBuilder);
 		}));
 
 	SequencerModule.GetObjectBindingContextMenuExtensibilityManager()->AddExtender(RebindComponentMenuExtender);
@@ -3441,45 +3435,48 @@ void ULevelSequenceEditorSubsystem::PopulateQuickBindSubMenu(FMenuBuilder& MenuB
 			FMovieSceneDirectorBlueprintEndpointDefinition EndpointDefinition
 			)
 			{
-				// Change or convert the binding
-				ChangeBindingTypes(Sequencer, BindingsToChange, [Sequencer, Sequence, bConvert, OnBindingChanged, CustomBindingType, &SelectedAction, InSelectionType, Blueprint, EndpointDefinition](FGuid BindingID, int32 BindingIndex)
-					{
-						FMovieScenePossessable* NewPossessable = nullptr;
-						if (FMovieSceneBindingReferences* BindingReferences = Sequence->GetBindingReferences())
+				if (!SelectedAction.IsEmpty())
+				{
+					// Change or convert the binding
+					ChangeBindingTypes(Sequencer, BindingsToChange, [Sequencer, Sequence, bConvert, OnBindingChanged, CustomBindingType, &SelectedAction, InSelectionType, Blueprint, EndpointDefinition](FGuid BindingID, int32 BindingIndex)
 						{
-							if (bConvert)
-							{
-								NewPossessable = FSequencerUtilities::ConvertToCustomBinding(Sequencer, BindingID, CustomBindingType, BindingIndex);
-							}
-							else
+							FMovieScenePossessable* NewPossessable = nullptr;
+							if (FMovieSceneBindingReferences* BindingReferences = Sequence->GetBindingReferences())
 							{
-								BindingReferences->AddOrReplaceBinding(BindingID,
-									NewObject<UMovieSceneCustomBinding>(Sequence->GetMovieScene(), CustomBindingType),
-									BindingIndex);
-								NewPossessable = Sequence->GetMovieScene()->FindPossessable(BindingID);
-							}
+								if (bConvert)
+								{
+									NewPossessable = FSequencerUtilities::ConvertToCustomBinding(Sequencer, BindingID, CustomBindingType, BindingIndex);
+								}
+								else
+								{
+									BindingReferences->AddOrReplaceBinding(BindingID,
+										NewObject<UMovieSceneCustomBinding>(Sequence->GetMovieScene(), CustomBindingType),
+										BindingIndex);
+									NewPossessable = Sequence->GetMovieScene()->FindPossessable(BindingID);
+								}
 
-							UMovieSceneCustomBinding* NewCustomBinding = BindingReferences->GetCustomBinding(BindingID, BindingIndex);
+								UMovieSceneCustomBinding* NewCustomBinding = BindingReferences->GetCustomBinding(BindingID, BindingIndex);
 
-							TArray<void*> RawData;
-							if (UMovieSceneReplaceableDirectorBlueprintBinding* ReplaceableBinding = Cast<UMovieSceneReplaceableDirectorBlueprintBinding>(NewCustomBinding))
-							{
-								RawData.Add(&ReplaceableBinding->DynamicBinding);
-							}
-							else if (UMovieSceneSpawnableDirectorBlueprintBinding* SpawnableBinding = Cast<UMovieSceneSpawnableDirectorBlueprintBinding>(NewCustomBinding))
-							{
-								RawData.Add(&SpawnableBinding->DynamicBinding);
-							}
+								TArray<void*> RawData;
+								if (UMovieSceneReplaceableDirectorBlueprintBinding* ReplaceableBinding = Cast<UMovieSceneReplaceableDirectorBlueprintBinding>(NewCustomBinding))
+								{
+									RawData.Add(&ReplaceableBinding->DynamicBinding);
+								}
+								else if (UMovieSceneSpawnableDirectorBlueprintBinding* SpawnableBinding = Cast<UMovieSceneSpawnableDirectorBlueprintBinding>(NewCustomBinding))
+								{
+									RawData.Add(&SpawnableBinding->DynamicBinding);
+								}
 
-							// Create temporary director blueprint binding customization for use in creating the endpoint
-							TSharedRef<FMovieSceneDynamicBindingCustomization> BlueprintBindingCustomization = StaticCastSharedRef<FMovieSceneDynamicBindingCustomization>(FMovieSceneDynamicBindingCustomization::MakeInstance(Sequence->GetMovieScene(), BindingID, BindingIndex));
-							BlueprintBindingCustomization->SetRawData(RawData);
-							BlueprintBindingCustomization->HandleQuickBindActionSelected(SelectedAction, InSelectionType, Blueprint, EndpointDefinition);
+								// Create temporary director blueprint binding customization for use in creating the endpoint
+								TSharedRef<FMovieSceneDynamicBindingCustomization> BlueprintBindingCustomization = StaticCastSharedRef<FMovieSceneDynamicBindingCustomization>(FMovieSceneDynamicBindingCustomization::MakeInstance(Sequence->GetMovieScene(), BindingID, BindingIndex));
+								BlueprintBindingCustomization->SetRawData(RawData);
+								BlueprintBindingCustomization->HandleQuickBindActionSelected(SelectedAction, InSelectionType, Blueprint, EndpointDefinition);
 
-						}
-						return NewPossessable;
-					},
-					OnBindingChanged);
+							}
+							return NewPossessable;
+						},
+						OnBindingChanged);
+				}
 			}));
 }
 
diff --git a/Engine/Source/Editor/MovieSceneTools/Private/Conditions/MovieSceneConditionCustomization.cpp b/Engine/Source/Editor/MovieSceneTools/Private/Conditions/MovieSceneConditionCustomization.cpp
--- a/Engine/Source/Editor/MovieSceneTools/Private/Conditions/MovieSceneConditionCustomization.cpp
+++ b/Engine/Source/Editor/MovieSceneTools/Private/Conditions/MovieSceneConditionCustomization.cpp
@@ -125,7 +125,7 @@ void FMovieSceneConditionCustomization::CustomizeChildren(TSharedRef<IPropertyHa
 	// Create new properties in the parent layout rather than adding a single item to a single category
 	IDetailLayoutBuilder& LayoutBuilder = ChildBuilder.GetParentCategory().GetParentLayout();
 
-	IDetailCategoryBuilder& NewCategory = LayoutBuilder.EditCategory(TEXT("Condition"), FText::GetEmpty(), ECategoryPriority::TypeSpecific);
+	//IDetailCategoryBuilder& NewCategory = LayoutBuilder.EditCategory(TEXT("Condition"), FText::GetEmpty(), ECategoryPriority::TypeSpecific);
 
 	// Hold onto a reference to the details view to prevent it from being destroyed immediately when the menu goes away. 
 	DetailsView = StaticCastSharedPtr<IDetailsView>(LayoutBuilder.GetDetailsView()->AsShared().ToSharedPtr());
@@ -151,7 +151,7 @@ void FMovieSceneConditionCustomization::CustomizeChildren(TSharedRef<IPropertyHa
 			{
 				continue;
 			}
-			NewCategory.AddProperty(ChildHandle);
+			ChildBuilder.AddProperty(ChildHandle);
 		}
 	}
 }
@@ -306,12 +306,15 @@ void FMovieSceneConditionCustomization::PopulateQuickBindSubMenu(FMenuBuilder& M
 			FMovieSceneDirectorBlueprintEndpointDefinition EndpointDefinition
 		)
 		{
-			// Create a new director blueprint condition and set it in the details view. Use 'interactive change' so we don't early fire the property finished changing event and reset the details view mid-change
-			PropertyCustomizationHelpers::CreateNewInstanceOfEditInlineObjectClass(SharedThis->ConditionPropertyHandle.ToSharedRef(), UMovieSceneDirectorBlueprintCondition::StaticClass(), EPropertyValueSetFlags::InteractiveChange);
-			TSharedPtr<IPropertyHandle> DirectorBlueprintConditionHandle = SharedThis->ConditionPropertyHandle->GetChildHandle(TEXT("DirectorBlueprintConditionData"));
-			BlueprintConditionCustomization->SetPropertyHandle(DirectorBlueprintConditionHandle);
-			BlueprintConditionCustomization->HandleQuickBindActionSelected(SelectedAction, InSelectionType, Blueprint, EndpointDefinition);
-			SharedThis->PropertyUtilities->ForceRefresh();
+			if (!SelectedAction.IsEmpty())
+			{
+				// Create a new director blueprint condition and set it in the details view. Use 'interactive change' so we don't early fire the property finished changing event and reset the details view mid-change
+				PropertyCustomizationHelpers::CreateNewInstanceOfEditInlineObjectClass(SharedThis->ConditionPropertyHandle.ToSharedRef(), UMovieSceneDirectorBlueprintCondition::StaticClass(), EPropertyValueSetFlags::InteractiveChange);
+				TSharedPtr<IPropertyHandle> DirectorBlueprintConditionHandle = SharedThis->ConditionPropertyHandle->GetChildHandle(TEXT("DirectorBlueprintConditionData"));
+				BlueprintConditionCustomization->SetPropertyHandle(DirectorBlueprintConditionHandle);
+				BlueprintConditionCustomization->HandleQuickBindActionSelected(SelectedAction, InSelectionType, Blueprint, EndpointDefinition);
+				SharedThis->PropertyUtilities->ForceRefresh();
+			}
 		}));
 	}
 }
diff --git a/Engine/Source/Editor/MovieSceneTools/Private/Conditions/MovieSceneDirectorBlueprintConditionCustomization.cpp b/Engine/Source/Editor/MovieSceneTools/Private/Conditions/MovieSceneDirectorBlueprintConditionCustomization.cpp
--- a/Engine/Source/Editor/MovieSceneTools/Private/Conditions/MovieSceneDirectorBlueprintConditionCustomization.cpp
+++ b/Engine/Source/Editor/MovieSceneTools/Private/Conditions/MovieSceneDirectorBlueprintConditionCustomization.cpp
@@ -44,7 +44,7 @@ TSharedRef<FMovieSceneDirectorBlueprintConditionCustomization> FMovieSceneDirect
 {
 	TSharedRef<FMovieSceneDirectorBlueprintConditionCustomization> Instance = MakeShared<FMovieSceneDirectorBlueprintConditionCustomization>();
 	Instance->EditedMovieScene = InMovieScene;
-	Instance->PropertyHandle = InPropertyHandle;
+	Instance->SetPropertyHandle(InPropertyHandle);
 	Instance->PropertyUtilities = InPropertyUtilities;
 	return Instance;
 }
diff --git a/Engine/Source/Editor/MovieSceneTools/Private/MovieSceneDirectorBlueprintEndpointCustomization.cpp b/Engine/Source/Editor/MovieSceneTools/Private/MovieSceneDirectorBlueprintEndpointCustomization.cpp
--- a/Engine/Source/Editor/MovieSceneTools/Private/MovieSceneDirectorBlueprintEndpointCustomization.cpp
+++ b/Engine/Source/Editor/MovieSceneTools/Private/MovieSceneDirectorBlueprintEndpointCustomization.cpp
@@ -261,10 +261,11 @@ void FMovieSceneDirectorBlueprintEndpointCustomization::CustomizeChildren(TShare
 
 	if (CommonFunction)
 	{
-		IDetailCategoryBuilder& DetailCategoryBuilder = ChildBuilder
-			.GetParentCategory()
+		IDetailCategoryBuilder& DetailCategoryBuilder = CreateNewCategoryForPayloadVariables() ? 
+			ChildBuilder.GetParentCategory()
 			.GetParentLayout()
-			.EditCategory("Payload", LOCTEXT("PayloadLabel", "Payload"), ECategoryPriority::Uncommon);
+			.EditCategory("Payload", LOCTEXT("PayloadLabel", "Payload"), ECategoryPriority::Uncommon) :
+			ChildBuilder.GetParentCategory();
 
 		bool bPayloadUpToDate = true;
 
@@ -346,7 +347,9 @@ void FMovieSceneDirectorBlueprintEndpointCustomization::CustomizeChildren(TShare
 					}
 				}
 
-				IDetailPropertyRow* ExternalRow = DetailCategoryBuilder.AddExternalStructureProperty(StructData.ToSharedRef(), Field->GetFName(), EPropertyLocation::Default, FAddPropertyParams().ForceShowProperty());
+				IDetailPropertyRow* ExternalRow = CreateNewCategoryForPayloadVariables() ? 
+				(DetailCategoryBuilder.AddExternalStructureProperty(StructData.ToSharedRef(), Field->GetFName(), EPropertyLocation::Default, FAddPropertyParams().ForceShowProperty()))
+				: (ChildBuilder.AddExternalStructureProperty(StructData.ToSharedRef(), Field->GetFName(), FAddPropertyParams().ForceShowProperty()));
 
 				TSharedPtr<IPropertyHandle> LocalVariableProperty = ExternalRow->GetPropertyHandle();
 				FSimpleDelegate Delegate = FSimpleDelegate::CreateSP(this, &FMovieSceneDirectorBlueprintEndpointCustomization::OnPayloadVariableChanged, StructData.ToSharedRef(), LocalVariableProperty);
@@ -360,7 +363,7 @@ void FMovieSceneDirectorBlueprintEndpointCustomization::CustomizeChildren(TShare
 
 		if (!bPayloadUpToDate)
 		{
-			DetailCategoryBuilder.AddCustomRow(FText())
+			(CreateNewCategoryForPayloadVariables() ? DetailCategoryBuilder.AddCustomRow(FText()) : ChildBuilder.AddCustomRow(FText()))
 			.WholeRowContent()
 			[
 				SNew(SHorizontalBox)
@@ -1052,6 +1055,16 @@ void FMovieSceneDirectorBlueprintEndpointCustomization::HandleQuickBindActionSel
 	}
 }
 
+void FMovieSceneDirectorBlueprintEndpointCustomization::SetPropertyHandle(TSharedPtr<IPropertyHandle> InPropertyHandle)
+{
+	PropertyHandle = InPropertyHandle;
+	if (PropertyHandle.IsValid())
+	{
+		PropertyRawData.Empty();
+		PropertyHandle->AccessRawData(PropertyRawData);
+	}
+}
+
 void FMovieSceneDirectorBlueprintEndpointCustomization::PopulateRebindSubMenu(FMenuBuilder& MenuBuilder, UMovieSceneSequence* Sequence)
 {
 	FMovieSceneSequenceEditor* SequenceEditor = FMovieSceneSequenceEditor::Find(Sequence);
diff --git a/Engine/Source/Editor/MovieSceneTools/Public/Conditions/MovieSceneDirectorBlueprintConditionCustomization.h b/Engine/Source/Editor/MovieSceneTools/Public/Conditions/MovieSceneDirectorBlueprintConditionCustomization.h
--- a/Engine/Source/Editor/MovieSceneTools/Public/Conditions/MovieSceneDirectorBlueprintConditionCustomization.h
+++ b/Engine/Source/Editor/MovieSceneTools/Public/Conditions/MovieSceneDirectorBlueprintConditionCustomization.h
@@ -42,7 +42,7 @@ protected:
 	virtual void OnSetEndpoint(UMovieSceneSequence* Sequence, UBlueprint* Blueprint, const TArray<UObject*> EditObjects, const TArray<void*> RawData, const FMovieSceneDirectorBlueprintEndpointDefinition& EndpointDefinition, UK2Node* NewEndpoint) override;
 	virtual void GetEditObjects(TArray<UObject*>& OutObjects) const override;
 	virtual void OnCollectQuickBindActions(UBlueprint* Blueprint, FBlueprintActionMenuBuilder& MenuBuilder) override;
-
+	virtual bool CreateNewCategoryForPayloadVariables() const override { return false; }
 private:
 
 	void SetEndpointImpl(UMovieScene* MovieScene, FMovieSceneDirectorBlueprintConditionData* DynamicBinding, UBlueprint* Blueprint, UK2Node* NewEndpoint);
diff --git a/Engine/Source/Editor/MovieSceneTools/Public/MovieSceneDirectorBlueprintEndpointCustomization.h b/Engine/Source/Editor/MovieSceneTools/Public/MovieSceneDirectorBlueprintEndpointCustomization.h
--- a/Engine/Source/Editor/MovieSceneTools/Public/MovieSceneDirectorBlueprintEndpointCustomization.h
+++ b/Engine/Source/Editor/MovieSceneTools/Public/MovieSceneDirectorBlueprintEndpointCustomization.h
@@ -70,7 +70,7 @@ public:
 	void HandleQuickBindActionSelected(const TArray<TSharedPtr<FEdGraphSchemaAction>>& SelectedAction, ESelectInfo::Type InSelectionType, UBlueprint* Blueprint, FMovieSceneDirectorBlueprintEndpointDefinition EndpointDefinition);
 
 	/* Set the property handle to use for customization. Used when using the customization class as a helper. */
-	void SetPropertyHandle(TSharedPtr<IPropertyHandle> InPropertyHandle) { PropertyHandle = InPropertyHandle;}
+	void SetPropertyHandle(TSharedPtr<IPropertyHandle> InPropertyHandle);
 
 	/* Set the raw data to use for customization. Used when the customization class is used as a helper without a details view. */
 	void SetRawData(const TArray<void*>& InRawData) { PropertyRawData = InRawData; }
@@ -198,6 +198,8 @@ protected:
 	 */
 	virtual void OnCollectAllRebindActions(UBlueprint* Blueprint, FBlueprintActionMenuBuilder& MenuBuilder) {}
 
+	virtual bool CreateNewCategoryForPayloadVariables() const { return true; }
+
 protected:
 
 	/**
