Sequencer: Expose SaveDefaultSpawnableState to python/BP

diff --git a/Engine/Plugins/MovieScene/LevelSequenceEditor/Source/LevelSequenceEditor/Private/LevelSequenceEditorSubsystem.cpp b/Engine/Plugins/MovieScene/LevelSequenceEditor/Source/LevelSequenceEditor/Private/LevelSequenceEditorSubsystem.cpp
--- a/Engine/Plugins/MovieScene/LevelSequenceEditor/Source/LevelSequenceEditor/Private/LevelSequenceEditorSubsystem.cpp
+++ b/Engine/Plugins/MovieScene/LevelSequenceEditor/Source/LevelSequenceEditor/Private/LevelSequenceEditorSubsystem.cpp
@@ -1077,11 +1077,28 @@ bool ULevelSequenceEditorSubsystem::ChangeActorTemplateClass(const FMovieSceneBi
 	TArray<FSequencerChangeBindingInfo> Bindings;
 	Bindings.Add(FSequencerChangeBindingInfo(ObjectBinding.BindingID, 0));
 
-	FSequencerUtilities::HandleTemplateActorClassPicked(ActorClass, Sequencer.ToSharedRef(), Bindings, [&bSuccess](){bSuccess=true;});
+	FSequencerUtilities::HandleTemplateActorClassPicked(ActorClass, Sequencer.ToSharedRef(), Bindings, [&bSuccess]() {bSuccess = true; });
 
 	return bSuccess;
 }
 
+void ULevelSequenceEditorSubsystem::SaveDefaultSpawnableState(const FMovieSceneBindingProxy& ObjectBinding)
+{
+	TSharedPtr<ISequencer> Sequencer = GetActiveSequencer();
+	if (Sequencer == nullptr)
+	{
+		return;
+	}
+
+	UMovieSceneSequence* Sequence = Sequencer->GetFocusedMovieSceneSequence();
+	if (!Sequence)
+	{
+		return;
+	}
+
+	Sequencer->GetSpawnRegister().SaveDefaultSpawnableState(ObjectBinding.BindingID, 0, Sequencer->GetFocusedTemplateID(), Sequencer->GetSharedPlaybackState());
+}
+
 void ULevelSequenceEditorSubsystem::CopyFolders(const TArray<UMovieSceneFolder*>& Folders, FString& ExportedText)
 {
 	FSequencerUtilities::CopyFolders(Folders, ExportedText);
diff --git a/Engine/Plugins/MovieScene/LevelSequenceEditor/Source/LevelSequenceEditor/Public/LevelSequenceEditorSubsystem.h b/Engine/Plugins/MovieScene/LevelSequenceEditor/Source/LevelSequenceEditor/Public/LevelSequenceEditorSubsystem.h
--- a/Engine/Plugins/MovieScene/LevelSequenceEditor/Source/LevelSequenceEditor/Public/LevelSequenceEditorSubsystem.h
+++ b/Engine/Plugins/MovieScene/LevelSequenceEditor/Source/LevelSequenceEditor/Public/LevelSequenceEditorSubsystem.h
@@ -160,6 +160,10 @@ public:
 	UFUNCTION(BlueprintCallable, Category = "Level Sequence Editor")
 	bool ChangeActorTemplateClass(const FMovieSceneBindingProxy& ObjectBinding, TSubclassOf<AActor> ActorClass);
 
+	/* Save the default state of the spawnable. */
+	UFUNCTION(BlueprintCallable, Category = "Level Sequence Editor")
+	void SaveDefaultSpawnableState(const FMovieSceneBindingProxy& ObjectBinding);
+
 	/** 
 	 * Copy folders 
 	 * The copied folders will be saved to the clipboard as well as assigned to the ExportedText string. 
