Sequencer- Hide condition editing on tracks that don't support ECS.

diff --git a/Engine/Plugins/MovieScene/LevelSequenceEditor/Source/LevelSequenceEditor/Private/LevelSequenceEditorSubsystem.cpp b/Engine/Plugins/MovieScene/LevelSequenceEditor/Source/LevelSequenceEditor/Private/LevelSequenceEditorSubsystem.cpp
--- a/Engine/Plugins/MovieScene/LevelSequenceEditor/Source/LevelSequenceEditor/Private/LevelSequenceEditorSubsystem.cpp
+++ b/Engine/Plugins/MovieScene/LevelSequenceEditor/Source/LevelSequenceEditor/Private/LevelSequenceEditorSubsystem.cpp
@@ -719,8 +719,8 @@ void ULevelSequenceEditorSubsystem::AddTrackRowMetadataCustomizations(TSharedRef
 		if (MovieScene)
 		{
 			// Although we normally customize this type, we need to do it instanced here to pass in the sequence information, as it won't be part of an outer sequence UObject
-			DetailsView->RegisterInstancedCustomPropertyTypeLayout("MovieSceneConditionContainer", FOnGetPropertyTypeCustomizationInstance::CreateLambda([=]() {
-				return FMovieSceneConditionCustomization::MakeInstance(Sequence); }));
+			DetailsView->RegisterInstancedCustomPropertyTypeLayout("MovieSceneConditionContainer", FOnGetPropertyTypeCustomizationInstance::CreateLambda([Sequence, WeakSequencer=ActiveSequencer.ToWeakPtr()]() {
+				return FMovieSceneConditionCustomization::MakeInstance(Sequence, WeakSequencer); }));
 
 			DetailsView->RegisterInstancedCustomPropertyTypeLayout("MovieSceneDirectorBlueprintConditionData", FOnGetPropertyTypeCustomizationInstance::CreateLambda([=]() {
 				return FMovieSceneDirectorBlueprintConditionCustomization::MakeInstance(MovieScene); }));
@@ -3190,9 +3190,9 @@ void ULevelSequenceEditorSubsystem::RefreshTrackRowMetadataDetails(IDetailsView*
 		UMovieSceneTrackRowMetadataHelper* Helper = TrackRowMetadataHelperList.Add_GetRef(NewObject<UMovieSceneTrackRowMetadataHelper>(this));
 		if (Track && Helper)
 		{
+			Helper->OwnerTrack = Track;
 			if (const FMovieSceneTrackRowMetadata* Metadata = Track->FindTrackRowMetadata(SelectedTrackRows[Index].Value))
 			{
-				Helper->OwnerTrack = Track;
 				Helper->TrackRowMetadata = *Metadata;
 				if (Metadata->ConditionContainer.Condition)
 				{
diff --git a/Engine/Source/Editor/MovieSceneTools/Private/Conditions/MovieSceneConditionCustomization.cpp b/Engine/Source/Editor/MovieSceneTools/Private/Conditions/MovieSceneConditionCustomization.cpp
--- a/Engine/Source/Editor/MovieSceneTools/Private/Conditions/MovieSceneConditionCustomization.cpp
+++ b/Engine/Source/Editor/MovieSceneTools/Private/Conditions/MovieSceneConditionCustomization.cpp
@@ -24,6 +24,7 @@
 #include "MovieScene.h"
 #include "ClassViewerFilter.h"
 #include "ScopedTransaction.h"
+#include "ISequencer.h"
 
 #define LOCTEXT_NAMESPACE "MovieSceneConditionCustomization"
 
@@ -71,10 +72,11 @@ TSharedRef<IPropertyTypeCustomization> FMovieSceneConditionCustomization::MakeIn
 	return Instance;
 }
 
-TSharedRef<IPropertyTypeCustomization> FMovieSceneConditionCustomization::MakeInstance(UMovieSceneSequence* InMovieSceneSequence)
+TSharedRef<IPropertyTypeCustomization> FMovieSceneConditionCustomization::MakeInstance(UMovieSceneSequence* InMovieSceneSequence, const TWeakPtr<ISequencer> Sequencer)
 {
 	TSharedRef<FMovieSceneConditionCustomization> Instance = MakeShared<FMovieSceneConditionCustomization>();
 	Instance->Sequence = InMovieSceneSequence;
+	Instance->Sequencer = Sequencer;
 	return Instance;
 }
 
@@ -94,7 +96,9 @@ void FMovieSceneConditionCustomization::CustomizeHeader(TSharedRef<IPropertyHand
 	}
 
 	// If conditions not allowed, hide condition property functionality
-	if (!Sequence.IsValid() || !Sequence->GetMovieScene()->IsConditionClassAllowed(UMovieSceneCondition::StaticClass()) || (Track.IsValid() && !Track->SupportsConditions()))
+	if (!Sequence.IsValid() 
+	|| !Sequence->GetMovieScene()->IsConditionClassAllowed(UMovieSceneCondition::StaticClass()) 
+	|| (Track.IsValid() && Sequencer.IsValid() && !Sequencer.Pin()->TrackSupportsConditions(Track.Get())))
 	{
 		ConditionContainerPropertyHandle->MarkHiddenByCustomization();
 		return;
diff --git a/Engine/Source/Editor/MovieSceneTools/Public/Conditions/MovieSceneConditionCustomization.h b/Engine/Source/Editor/MovieSceneTools/Public/Conditions/MovieSceneConditionCustomization.h
--- a/Engine/Source/Editor/MovieSceneTools/Public/Conditions/MovieSceneConditionCustomization.h
+++ b/Engine/Source/Editor/MovieSceneTools/Public/Conditions/MovieSceneConditionCustomization.h
@@ -20,6 +20,7 @@ class UMovieSceneSequence;
 class UMovieSceneTrack;
 class FMovieSceneDirectorBlueprintConditionCustomization;
 class IDetailsView;
+class ISequencer;
 
 // Helper UObject for editing optional track row metadata not in-place. A UObject instead of a UStruct because we need to support instanced sub objects (conditions)
 UCLASS(CollapseCategories)
@@ -41,7 +42,7 @@ class FMovieSceneConditionCustomization : public IPropertyTypeCustomization
 public:
 
 	static MOVIESCENETOOLS_API TSharedRef<IPropertyTypeCustomization> MakeInstance();
-	static MOVIESCENETOOLS_API TSharedRef<IPropertyTypeCustomization> MakeInstance(UMovieSceneSequence* Sequence);
+	static MOVIESCENETOOLS_API TSharedRef<IPropertyTypeCustomization> MakeInstance(UMovieSceneSequence* Sequence, const TWeakPtr<ISequencer> Sequencer);
 
 	// IPropertyTypeCustomization interface
 	virtual void CustomizeHeader(TSharedRef<IPropertyHandle> InPropertyHandle, FDetailWidgetRow& HeaderRow, IPropertyTypeCustomizationUtils& CustomizationUtils) override;
@@ -137,4 +138,6 @@ private:
 
 	TWeakObjectPtr<UMovieSceneSequence> Sequence;
 	TWeakObjectPtr<UMovieSceneTrack> Track;
+
+	TWeakPtr<ISequencer> Sequencer;
 };
diff --git a/Engine/Source/Editor/Sequencer/Private/Sequencer.cpp b/Engine/Source/Editor/Sequencer/Private/Sequencer.cpp
--- a/Engine/Source/Editor/Sequencer/Private/Sequencer.cpp
+++ b/Engine/Source/Editor/Sequencer/Private/Sequencer.cpp
@@ -8548,6 +8548,32 @@ TArray<FText> FSequencer::GetTrackFilterNames() const
 	return FilterBar->GetFilterDisplayNames();
 }
 
+bool FSequencer::TrackSupportsConditions(const UMovieSceneTrack* Track) const
+{
+	if (Track)
+	{
+		if (!Track->SupportsConditions())
+		{
+			return false;
+		}
+
+		// Non ECS tracks don't support conditions
+		if (CompiledDataManager)
+		{
+			if (const FMovieSceneEvaluationTemplate* Template = CompiledDataManager->FindTrackTemplate(CompiledDataManager->FindDataID(Track->GetTypedOuter<UMovieSceneSequence>())))
+			{
+				if (Template->FindTrack(Track->GetSignature()))
+				{
+					return false;
+				}
+			}
+		}
+
+		return true;
+	}
+	return false;
+}
+
 void FSequencer::ToggleNodeLocked()
 {
 	using namespace UE::Sequencer;
diff --git a/Engine/Source/Editor/Sequencer/Private/Sequencer.h b/Engine/Source/Editor/Sequencer/Private/Sequencer.h
--- a/Engine/Source/Editor/Sequencer/Private/Sequencer.h
+++ b/Engine/Source/Editor/Sequencer/Private/Sequencer.h
@@ -823,6 +823,8 @@ public:
 	virtual bool IsTrackFilterEnabled(const FText& InFilterName) const override;
 	virtual TArray<FText> GetTrackFilterNames() const override;
 
+	virtual bool TrackSupportsConditions(const UMovieSceneTrack* Track) const override;
+
 public:
 
 	// IMovieScenePlayer interface
diff --git a/Engine/Source/Editor/Sequencer/Public/ISequencer.h b/Engine/Source/Editor/Sequencer/Public/ISequencer.h
--- a/Engine/Source/Editor/Sequencer/Public/ISequencer.h
+++ b/Engine/Source/Editor/Sequencer/Public/ISequencer.h
@@ -940,6 +940,8 @@ public:
 	 */
 	virtual bool UnregisterDrawerSection(const FName InDrawerId, const FName InSectionId) = 0;
 
+	virtual bool TrackSupportsConditions(const UMovieSceneTrack* Track) const = 0;
+
 protected:
 	FOnInitializeDetailsPanel InitializeDetailsPanelEvent;
 	FOnCameraAddedToSequencer CameraAddedToSequencer;
diff --git a/Engine/Source/Runtime/MovieScene/Private/MovieSceneTrack.cpp b/Engine/Source/Runtime/MovieScene/Private/MovieSceneTrack.cpp
--- a/Engine/Source/Runtime/MovieScene/Private/MovieSceneTrack.cpp
+++ b/Engine/Source/Runtime/MovieScene/Private/MovieSceneTrack.cpp
@@ -19,6 +19,7 @@
 
 #include "Channels/MovieSceneChannelProxy.h"
 #include "Channels/MovieSceneChannel.h"
+#include "Compilation/IMovieSceneTrackTemplateProducer.h"
 
 #include UE_INLINE_GENERATED_CPP_BY_NAME(MovieSceneTrack)
 
diff --git a/Engine/Source/Runtime/MovieSceneTracks/Public/Tracks/MovieSceneAudioTrack.h b/Engine/Source/Runtime/MovieSceneTracks/Public/Tracks/MovieSceneAudioTrack.h
--- a/Engine/Source/Runtime/MovieSceneTracks/Public/Tracks/MovieSceneAudioTrack.h
+++ b/Engine/Source/Runtime/MovieSceneTracks/Public/Tracks/MovieSceneAudioTrack.h
@@ -5,7 +5,6 @@
 #include "CoreMinimal.h"
 #include "UObject/ObjectMacros.h"
 #include "MovieSceneNameableTrack.h"
-#include "Compilation/IMovieSceneTrackTemplateProducer.h"
 #include "MovieSceneAudioTrack.generated.h"
 
 class USoundBase;
diff --git a/Engine/Source/Runtime/MovieSceneTracks/Public/Tracks/MovieSceneEulerTransformTrack.h b/Engine/Source/Runtime/MovieSceneTracks/Public/Tracks/MovieSceneEulerTransformTrack.h
--- a/Engine/Source/Runtime/MovieSceneTracks/Public/Tracks/MovieSceneEulerTransformTrack.h
+++ b/Engine/Source/Runtime/MovieSceneTracks/Public/Tracks/MovieSceneEulerTransformTrack.h
@@ -5,7 +5,6 @@
 #include "CoreMinimal.h"
 #include "UObject/ObjectMacros.h"
 #include "Tracks/MovieScenePropertyTrack.h"
-#include "Compilation/IMovieSceneTrackTemplateProducer.h"
 #include "MovieSceneEulerTransformTrack.generated.h"
 
 /**
