Transform baking: evaluate constraints when baking

#rnx

diff --git a/Engine/Plugins/MovieScene/LevelSequenceEditor/Source/LevelSequenceEditor/LevelSequenceEditor.Build.cs b/Engine/Plugins/MovieScene/LevelSequenceEditor/Source/LevelSequenceEditor/LevelSequenceEditor.Build.cs
--- a/Engine/Plugins/MovieScene/LevelSequenceEditor/Source/LevelSequenceEditor/LevelSequenceEditor.Build.cs
+++ b/Engine/Plugins/MovieScene/LevelSequenceEditor/Source/LevelSequenceEditor/LevelSequenceEditor.Build.cs
@@ -34,6 +34,7 @@ public class LevelSequenceEditor : ModuleRules
 				"BlueprintGraph",
                 "CinematicCamera",
 				"ClassViewer",
+				"Constraints",
 				"Core",
 				"CoreUObject",
 				"CurveEditor",
diff --git a/Engine/Plugins/MovieScene/LevelSequenceEditor/Source/LevelSequenceEditor/Private/LevelSequenceEditorSubsystem.cpp b/Engine/Plugins/MovieScene/LevelSequenceEditor/Source/LevelSequenceEditor/Private/LevelSequenceEditorSubsystem.cpp
--- a/Engine/Plugins/MovieScene/LevelSequenceEditor/Source/LevelSequenceEditor/Private/LevelSequenceEditorSubsystem.cpp
+++ b/Engine/Plugins/MovieScene/LevelSequenceEditor/Source/LevelSequenceEditor/Private/LevelSequenceEditorSubsystem.cpp
@@ -1794,6 +1794,10 @@ bool ULevelSequenceEditorSubsystem::BakeTransformWithSettings(const TArray<FMovi
 	TArray<FFrameNumber> AllFrames;
 	TotalFrameMap.GenerateKeyArray(AllFrames);
 
+	UWorld* PlaybackContext = Sequencer->GetPlaybackContext()->GetWorld();
+	ensure(PlaybackContext);
+	const FConstraintsManagerController& Controller = FConstraintsManagerController::Get(PlaybackContext);
+
 	for (FFrameNumber KeyTime: AllFrames)
 	{
 		TOptional<FFrameTime> NewGlobalTime = LocalToRootTransform.TryTransformTime(KeyTime);
@@ -1803,6 +1807,7 @@ bool ULevelSequenceEditorSubsystem::BakeTransformWithSettings(const TArray<FMovi
 		}
 
 		Sequencer->SetGlobalTime(NewGlobalTime.GetValue());
+		Controller.EvaluateAllConstraints();
 
 		for (const FMovieSceneBindingProxy& ObjectBinding : ObjectBindings)
 		{
