Sequencer: Fix crash when converting an empty spawnable to a possessable. This was actually an issue in general with changing binding type from spawnable in the binding properties. When changing the binding type vs. 'converting' binding type, it resets the binding completely. This means that any spawnable template is set to null. However, it wasn't correctly updating the Sequence, and so the spawned object persisted. This caused confusion if you then used a convert function to convert to possessable as there was a mismatch between the object template and the spawned actor. This change addresses this by properly updating the sequence when changing the binding type, which will despawn any spawned acor and force re-evaluate the Sequence.

diff --git a/Engine/Plugins/MovieScene/LevelSequenceEditor/Source/LevelSequenceEditor/Private/LevelSequenceEditorSubsystem.cpp b/Engine/Plugins/MovieScene/LevelSequenceEditor/Source/LevelSequenceEditor/Private/LevelSequenceEditorSubsystem.cpp
--- a/Engine/Plugins/MovieScene/LevelSequenceEditor/Source/LevelSequenceEditor/Private/LevelSequenceEditorSubsystem.cpp
+++ b/Engine/Plugins/MovieScene/LevelSequenceEditor/Source/LevelSequenceEditor/Private/LevelSequenceEditorSubsystem.cpp
@@ -83,6 +83,7 @@
 #include "AssetToolsModule.h"
 #include "Kismet2/KismetEditorUtilities.h"
 #include "Kismet2/BlueprintEditorUtils.h"
+#include "MovieSceneSequencePlayer.h"
 
 #include UE_INLINE_GENERATED_CPP_BY_NAME(LevelSequenceEditorSubsystem)
 
@@ -225,11 +226,12 @@ TSharedRef<SWidget> FMovieSceneBindingPropertyInfoListCustomization::GenerateBin
 		{
 			Bindings.Add({ BindingGuid, BindingIndex });
 		}
-		LevelSequenceEditorSubsystem->AddChangeBindingTypeMenu(MenuBuilder, Sequencer, Bindings, false, [this, &DetailBuilder]()
+		LevelSequenceEditorSubsystem->AddChangeBindingTypeMenu(MenuBuilder, Sequencer, Bindings, false, [this, Sequencer, &DetailBuilder]()
 			{
 				if (IDetailsView* DetailsView = DetailBuilder.GetDetailsView())
 				{
 					LevelSequenceEditorSubsystem->RefreshBindingDetails(DetailsView, BindingGuid);
+					LevelSequenceEditorSubsystem->OnFinishedChangingLocators(FPropertyChangedEvent(nullptr), StaticCastSharedRef<IDetailsView>(DetailsView->AsShared()), BindingGuid);
 				}
 			});
 	}
diff --git a/Engine/Plugins/MovieScene/LevelSequenceEditor/Source/LevelSequenceEditor/Public/LevelSequenceEditorSubsystem.h b/Engine/Plugins/MovieScene/LevelSequenceEditor/Source/LevelSequenceEditor/Public/LevelSequenceEditorSubsystem.h
--- a/Engine/Plugins/MovieScene/LevelSequenceEditor/Source/LevelSequenceEditor/Public/LevelSequenceEditorSubsystem.h
+++ b/Engine/Plugins/MovieScene/LevelSequenceEditor/Source/LevelSequenceEditor/Public/LevelSequenceEditorSubsystem.h
@@ -367,6 +367,8 @@ private:
 	TSharedPtr<FExtender> BindingPropertiesMenuExtender;
 	TSharedPtr<FExtender> RebindComponentMenuExtender;
 	TSharedPtr<FExtender> SidebarMenuExtender;
+
+	friend class FMovieSceneBindingPropertyInfoListCustomization;
 };
 
 #if UE_ENABLE_INCLUDE_ORDER_DEPRECATED_IN_5_2
