Sequencer: Disable changing binding type directly unless there currently is not a bound object. Converting binding types is still allowed, just not resetting binding type to a new one. This should remove confusion/bug reports of losing the bound object when using the change menu rather than the conversion menu, while still allowing for empty bindings to be created and then having their binding type set afterwards.

diff --git a/Engine/Plugins/MovieScene/LevelSequenceEditor/Source/LevelSequenceEditor/Private/LevelSequenceEditorSubsystem.cpp b/Engine/Plugins/MovieScene/LevelSequenceEditor/Source/LevelSequenceEditor/Private/LevelSequenceEditorSubsystem.cpp
--- a/Engine/Plugins/MovieScene/LevelSequenceEditor/Source/LevelSequenceEditor/Private/LevelSequenceEditorSubsystem.cpp
+++ b/Engine/Plugins/MovieScene/LevelSequenceEditor/Source/LevelSequenceEditor/Private/LevelSequenceEditorSubsystem.cpp
@@ -159,6 +159,7 @@ void FMovieSceneBindingPropertyInfoListCustomization::CustomizeDetails(IDetailLa
 							bShowConvert = false;
 						}
 					}
+					bool bHasBoundObject = MovieSceneHelpers::GetSingleBoundObject(Sequence, BindingGuid, Sequencer->GetSharedPlaybackState(), 0) != nullptr;
 
 					FDetailWidgetRow& BindingTypeRow = SectionCategory.AddCustomRow(FText::GetEmpty());
 					BindingTypeRow.NameContent()
@@ -191,10 +192,11 @@ void FMovieSceneBindingPropertyInfoListCustomization::CustomizeDetails(IDetailLa
 									.Text(this, &FMovieSceneBindingPropertyInfoListCustomization::GetBindingTypeValueAsString)
 								]
 							]
+							.IsEnabled(!bHasBoundObject)
 						];
 
 					// Only show certain menus if we have a currently bound object
-					if (bShowConvert && MovieSceneHelpers::GetSingleBoundObject(Sequence, BindingGuid, Sequencer->GetSharedPlaybackState(), 0) != nullptr)
+					if (bShowConvert && bHasBoundObject)
 					{
 						FDetailWidgetRow& ConvertToRow = SectionCategory.AddCustomRow(LOCTEXT("ConvertBindingTo", "Convert Binding(s) To..."));
 						ConvertToRow.WholeRowContent()
