Sequencer: Hook up AssetEditorToolkit::SaveAsset_Execute so that Ctr+S in the level sequence UI will do the specific sequence saving operations (thumbnail, save all sequences from the root)

diff --git a/Engine/Plugins/MovieScene/LevelSequenceEditor/Source/LevelSequenceEditor/Private/LevelSequenceEditorToolkit.cpp b/Engine/Plugins/MovieScene/LevelSequenceEditor/Source/LevelSequenceEditor/Private/LevelSequenceEditorToolkit.cpp
--- a/Engine/Plugins/MovieScene/LevelSequenceEditor/Source/LevelSequenceEditor/Private/LevelSequenceEditorToolkit.cpp
+++ b/Engine/Plugins/MovieScene/LevelSequenceEditor/Source/LevelSequenceEditor/Private/LevelSequenceEditorToolkit.cpp
@@ -886,6 +886,11 @@ TSharedRef<SDockTab> FLevelSequenceEditorToolkit::HandleTabManagerSpawnTab(const
 		];
 }
 
+void FLevelSequenceEditorToolkit::SaveAsset_Execute()
+{
+	GetSequencer()->Save();
+}
+
 void FLevelSequenceEditorToolkit::OnClose()
 {
 	UWorld* World = PlaybackContext->GetPlaybackContext()->GetWorld();
diff --git a/Engine/Plugins/MovieScene/LevelSequenceEditor/Source/LevelSequenceEditor/Private/LevelSequenceEditorToolkit.h b/Engine/Plugins/MovieScene/LevelSequenceEditor/Source/LevelSequenceEditor/Private/LevelSequenceEditorToolkit.h
--- a/Engine/Plugins/MovieScene/LevelSequenceEditor/Source/LevelSequenceEditor/Private/LevelSequenceEditorToolkit.h
+++ b/Engine/Plugins/MovieScene/LevelSequenceEditor/Source/LevelSequenceEditor/Private/LevelSequenceEditorToolkit.h
@@ -88,6 +88,7 @@ public:
 		return TEXT("FLevelSequenceEditorToolkit");
 	}
 
+	virtual void SaveAsset_Execute() override;
 	virtual void OnClose() override;
 	virtual bool CanFindInContentBrowser() const override;
 
