Sequencer: Added GetPlaybackStartPosition and GetPlaybackEndPosition which return the values that the playback controls use

This allows python/BP to jump to the last frame, like the playback controls do, ie. minus one frame.

diff --git a/Engine/Plugins/MovieScene/LevelSequenceEditor/Source/LevelSequenceEditor/Private/LevelSequenceEditorBlueprintLibrary.cpp b/Engine/Plugins/MovieScene/LevelSequenceEditor/Source/LevelSequenceEditor/Private/LevelSequenceEditorBlueprintLibrary.cpp
--- a/Engine/Plugins/MovieScene/LevelSequenceEditor/Source/LevelSequenceEditor/Private/LevelSequenceEditorBlueprintLibrary.cpp
+++ b/Engine/Plugins/MovieScene/LevelSequenceEditor/Source/LevelSequenceEditor/Private/LevelSequenceEditorBlueprintLibrary.cpp
@@ -333,6 +333,50 @@ void ULevelSequenceEditorBlueprintLibrary::PlayTo(FMovieSceneSequencePlaybackPar
 	}
 }
 
+FMovieSceneSequencePlaybackParams ULevelSequenceEditorBlueprintLibrary::GetPlaybackStartPosition(EMovieSceneTimeUnit TimeUnit)
+{
+	FMovieSceneSequencePlaybackParams Params;
+
+	if (CurrentSequencer.IsValid())
+	{
+		FFrameRate DisplayRate = CurrentSequencer.Pin()->GetFocusedDisplayRate();
+		FFrameRate TickResolution = CurrentSequencer.Pin()->GetFocusedTickResolution();
+
+		UMovieSceneSequence* Sequence = CurrentSequencer.Pin()->GetFocusedMovieSceneSequence();
+		TRange<FFrameNumber> PlaybackRange = Sequence->GetMovieScene()->GetPlaybackRange();
+		FQualifiedFrameTime StartPosition(UE::MovieScene::DiscreteInclusiveLower(PlaybackRange), TickResolution);
+
+		Params.Frame = TimeUnit == EMovieSceneTimeUnit::DisplayRate ? ConvertFrameTime(StartPosition.Time, TickResolution, DisplayRate) : StartPosition.Time;
+		Params.Timecode = StartPosition.ToTimecode();
+		Params.Time = StartPosition.AsSeconds();
+	}
+
+	return Params;
+}
+
+FMovieSceneSequencePlaybackParams ULevelSequenceEditorBlueprintLibrary::GetPlaybackEndPosition(EMovieSceneTimeUnit TimeUnit)
+{
+	FMovieSceneSequencePlaybackParams Params;
+
+	if (CurrentSequencer.IsValid())
+	{
+		FFrameRate DisplayRate = CurrentSequencer.Pin()->GetFocusedDisplayRate();
+		FFrameRate TickResolution = CurrentSequencer.Pin()->GetFocusedTickResolution();
+
+		UMovieSceneSequence* Sequence = CurrentSequencer.Pin()->GetFocusedMovieSceneSequence();
+		TRange<FFrameNumber> PlaybackRange = Sequence->GetMovieScene()->GetPlaybackRange();
+
+		FFrameTime OneFrame = FFrameRate::TransformTime(FFrameTime(1), DisplayRate, TickResolution);
+		FQualifiedFrameTime EndPosition(UE::MovieScene::DiscreteExclusiveUpper(PlaybackRange) - OneFrame, TickResolution);
+
+		Params.Frame = TimeUnit == EMovieSceneTimeUnit::DisplayRate ? ConvertFrameTime(EndPosition.Time, TickResolution, DisplayRate) : EndPosition.Time;
+		Params.Timecode = EndPosition.ToTimecode();
+		Params.Time = EndPosition.AsSeconds();
+	}
+
+	return Params;
+}
+
 bool ULevelSequenceEditorBlueprintLibrary::IsPlaying()
 {
 	if (CurrentSequencer.IsValid())
diff --git a/Engine/Plugins/MovieScene/LevelSequenceEditor/Source/LevelSequenceEditor/Public/LevelSequenceEditorBlueprintLibrary.h b/Engine/Plugins/MovieScene/LevelSequenceEditor/Source/LevelSequenceEditor/Public/LevelSequenceEditorBlueprintLibrary.h
--- a/Engine/Plugins/MovieScene/LevelSequenceEditor/Source/LevelSequenceEditor/Public/LevelSequenceEditorBlueprintLibrary.h
+++ b/Engine/Plugins/MovieScene/LevelSequenceEditor/Source/LevelSequenceEditor/Public/LevelSequenceEditorBlueprintLibrary.h
@@ -155,6 +155,14 @@ public:
 
 public:
 
+	/** Return the playback start position */
+	UFUNCTION(BlueprintCallable, Category = "Level Sequence Editor", DisplayName = "Get Playback Start Time")
+	static FMovieSceneSequencePlaybackParams GetPlaybackStartPosition(EMovieSceneTimeUnit TimeUnit = EMovieSceneTimeUnit::DisplayRate);
+
+	/** Return end of the playback range in the Sequencer UI, which accounts for the exclusive upper bound */
+	UFUNCTION(BlueprintCallable, Category = "Level Sequence Editor", DisplayName = "Get Playback End Time")
+	static FMovieSceneSequencePlaybackParams GetPlaybackEndPosition(EMovieSceneTimeUnit TimeUnit = EMovieSceneTimeUnit::DisplayRate);
+
 	/** Check whether the sequence is actively playing. */
 	UFUNCTION(BlueprintPure, Category = "Level Sequence Editor")
 	static bool IsPlaying();
