Add "Toggle Limit Viewport Selection" to the new level editor viewport toolbar when opening the Sequencer

diff --git a/Engine/Plugins/Animation/ControlRig/Source/ControlRigEditor/Private/Editor/ControlRigViewportToolbarExtensions.cpp b/Engine/Plugins/Animation/ControlRig/Source/ControlRigEditor/Private/Editor/ControlRigViewportToolbarExtensions.cpp
--- a/Engine/Plugins/Animation/ControlRig/Source/ControlRigEditor/Private/Editor/ControlRigViewportToolbarExtensions.cpp
+++ b/Engine/Plugins/Animation/ControlRig/Source/ControlRigEditor/Private/Editor/ControlRigViewportToolbarExtensions.cpp
@@ -108,12 +108,11 @@ void UE::ControlRig::PopulateControlRigViewportToolbarTransformSubmenu(const FNa
 	{
 		FToolMenuSection& SelectionSection = Menu->FindOrAddSection("Selection");
 
-		UControlRigEditModeSettings* const Settings = GetMutableDefault<UControlRigEditModeSettings>();
-
-		// Add "Restore Coordinate Space on Switch" checkbox.
+		// Add "Select Only Control Rig Controls" entry.
 		{
-			FUIAction Action;
+			UControlRigEditModeSettings* const Settings = GetMutableDefault<UControlRigEditModeSettings>();
 
+			FUIAction Action;
 			Action.ExecuteAction = FExecuteAction::CreateLambda(
 				[Settings]() -> void
 				{
diff --git a/Engine/Plugins/MovieScene/LevelSequenceEditor/Source/LevelSequenceEditor/Private/LevelSequenceEditorToolkit.cpp b/Engine/Plugins/MovieScene/LevelSequenceEditor/Source/LevelSequenceEditor/Private/LevelSequenceEditorToolkit.cpp
--- a/Engine/Plugins/MovieScene/LevelSequenceEditor/Source/LevelSequenceEditor/Private/LevelSequenceEditorToolkit.cpp
+++ b/Engine/Plugins/MovieScene/LevelSequenceEditor/Source/LevelSequenceEditor/Private/LevelSequenceEditorToolkit.cpp
@@ -17,6 +17,7 @@
 #include "LevelEditor.h"
 #include "LevelEditorViewport.h"
 #include "CineCameraActor.h"
+#include "EditorModeManager.h"
 #include "KeyPropertyParams.h"
 #include "ISequencerModule.h"
 #include "LevelSequenceActor.h"
@@ -29,7 +30,9 @@
 #include "Tracks/IMovieSceneTransformOrigin.h"
 #include "MovieSceneToolsProjectSettings.h"
 #include "MovieSceneToolHelpers.h"
+#include "SEditorViewport.h"
 #include "ScopedTransaction.h"
+#include "SequencerCommands.h"
 #include "ToolMenuEntry.h"
 #include "Widgets/Docking/SDockTab.h"
 #include "SequencerSettings.h"
@@ -40,6 +43,7 @@
 #include "ClassViewerFilter.h"
 #include "ClassViewerModule.h"
 #include "SequencerUtilities.h"
+#include "ViewportToolbar/UnrealEdViewportToolbarContext.h"
 
 // @todo sequencer: hack: setting defaults for transform tracks
 
@@ -52,6 +56,72 @@
 
 #define LOCTEXT_NAMESPACE "LevelSequenceEditor"
 
+namespace UE::MovieScene::Private
+{
+
+const FName ViewportToolbarOwnerName = "LevelSequenceEditorViewportToolbar";
+
+FToolMenuEntry CreateToggleViewportSelectionEntry(TWeakPtr<ISequencer> InWeakSequencer)
+{
+	return FToolMenuEntry::InitDynamicEntry(
+		"DynamicToggleViewportSelection",
+		FNewToolMenuSectionDelegate::CreateLambda(
+			[InWeakSequencer](FToolMenuSection& InDynamicSection)
+			{
+				// First attempt to use the non-context Sequencer.
+				TSharedPtr<ISequencer> Sequencer = InWeakSequencer.Pin();
+				if (!Sequencer)
+				{
+					return;
+				}
+
+				if (!Sequencer->GetHostCapabilities().bSupportsViewportSelectability)
+				{
+					return;
+				}
+
+				const TSharedPtr<FUICommandList> CommandList =
+					Sequencer->GetCommandBindings(ESequencerCommandBindings::Sequencer);
+
+				FToolMenuEntry& Entry = InDynamicSection.AddMenuEntryWithCommandList(
+					FSequencerCommands::Get().ToggleLimitViewportSelection, CommandList
+				);
+				Entry.InsertPosition.Position = EToolMenuInsertType::Last;
+
+				if (UUnrealEdViewportToolbarContext* ViewportToolbarContext =
+						InDynamicSection.FindContext<UUnrealEdViewportToolbarContext>())
+				{
+					// Show the entry in the top-level toolbar if we're in Animation Mode (EditMode.ControlRig). Otherwise it will be in the Transform submenu.
+					Entry.SetShowInToolbarTopLevel(TAttribute<bool>::CreateLambda(
+						[WeakViewport = ViewportToolbarContext->Viewport]() -> bool
+						{
+							if (TSharedPtr<SEditorViewport> Viewport = WeakViewport.Pin())
+							{
+								if (TSharedPtr<FEditorViewportClient> Client = Viewport->GetViewportClient())
+								{
+									if (FEditorModeTools* ModeTools = Client->GetModeTools())
+									{
+										// Hard-code the mode name instead of using FControlRigEditMode::ModeName to avoid adding a dependency on ControlRigEditor.
+										const FName ControlRigModeName = "EditMode.ControlRig";
+
+										if (ModeTools->GetActiveMode(ControlRigModeName))
+										{
+											return true;
+										}
+									}
+								}
+							}
+
+							return false;
+						}
+					));
+				}
+			}
+		)
+	);
+}
+
+}; // namespace UE::MovieScene::Private
 
 /* Local constants
  *****************************************************************************/
@@ -228,6 +298,23 @@ void FLevelSequenceEditorToolkit::Initialize(const EToolkitMode::Type Mode, cons
 			VRMode->RefreshVREditorSequencer(Sequencer.Get());
 		}
 	}
+
+	// Extend the Level Editor viewport toolbar.
+	{
+		FToolMenuOwnerScoped ScopeOwner(UE::MovieScene::Private::ViewportToolbarOwnerName);
+
+		const FName TransformMenu = "LevelEditor.ViewportToolbar.Transform";
+		UToolMenu* const Menu = UToolMenus::Get()->ExtendMenu(TransformMenu);
+
+		FToolMenuSection& SelectionSection = Menu->FindOrAddSection("Selection");
+
+		{
+			FToolMenuEntry Entry = UE::MovieScene::Private::CreateToggleViewportSelectionEntry(Sequencer.ToWeakPtr());
+			Entry.InsertPosition.Position = EToolMenuInsertType::Last;
+
+			SelectionSection.AddEntry(Entry);
+		}
+	}
 }
 
 
@@ -893,6 +980,9 @@ void FLevelSequenceEditorToolkit::SaveAsset_Execute()
 
 void FLevelSequenceEditorToolkit::OnClose()
 {
+	// Remove the viewport toolbar extensions we added earlier.
+	UToolMenus::Get()->UnregisterOwnerByName(UE::MovieScene::Private::ViewportToolbarOwnerName);
+
 	UWorld* World = PlaybackContext->GetPlaybackContext()->GetWorld();
 	UVREditorMode* VRMode = Cast<UVREditorMode>(GEditor->GetEditorWorldExtensionsManager()->GetEditorWorldExtensions(World)->FindExtension(UVREditorMode::StaticClass()));
 	if (VRMode != nullptr)
diff --git a/Engine/Source/Editor/EditorStyle/Private/StarshipStyle.cpp b/Engine/Source/Editor/EditorStyle/Private/StarshipStyle.cpp
--- a/Engine/Source/Editor/EditorStyle/Private/StarshipStyle.cpp
+++ b/Engine/Source/Editor/EditorStyle/Private/StarshipStyle.cpp
@@ -2296,6 +2296,8 @@ void FStarshipEditorStyle::FStyle::SetupSequencerStyles()
 		Set("Sequencer.SetKeyGroup", new IMAGE_BRUSH_SVG("Starship/Sequencer/KeyGroup", Icon20x20));
 		Set("Sequencer.SetKeyChanged", new IMAGE_BRUSH_SVG("Starship/Sequencer/KeyChanged", Icon20x20));
 		Set("Sequencer.ToggleIsSnapEnabled", new IMAGE_BRUSH_SVG("Starship/Sequencer/Snap", Icon20x20));
+		Set("Sequencer.ToggleLimitViewportSelection",
+			new IMAGE_BRUSH_SVG("Starship/Sequencer/SelectOnlyInSequence_16", Icon20x20));
 
 		Set("Sequencer.OpenTaggedBindingManager", new IMAGE_BRUSH("Sequencer/Main_Icons/Icon_Sequencer_OpenTaggedBindingManager_16x", Icon48x48));
 		Set("Sequencer.OpenNodeGroupsManager", new IMAGE_BRUSH("Sequencer/Main_Icons/Icon_Sequencer_OpenGroupManager_16x", Icon48x48));
diff --git a/Engine/Source/Editor/Sequencer/Private/SSequencer.cpp b/Engine/Source/Editor/Sequencer/Private/SSequencer.cpp
--- a/Engine/Source/Editor/Sequencer/Private/SSequencer.cpp
+++ b/Engine/Source/Editor/Sequencer/Private/SSequencer.cpp
@@ -1724,10 +1724,8 @@ void SSequencer::PopulateToolBar(UToolMenu* InMenu)
 		{
 			FToolMenuSection& SelectionSection = InMenu->AddSection(TEXT("Selection"));
 
-			FToolMenuEntry ToggleLockViewportSelectionEntry = FToolMenuEntry::InitToolBarButton(FSequencerCommands::Get().ToggleLimitViewportSelection
-				, FText()
-				, FSequencerCommands::Get().ToggleLimitViewportSelection->GetDescription()
-				, FSlateIcon(FAppStyle::GetAppStyleSetName(), TEXT("ContentBrowser.Local.Small")));
+			FToolMenuEntry ToggleLockViewportSelectionEntry =
+				FToolMenuEntry::InitToolBarButton(FSequencerCommands::Get().ToggleLimitViewportSelection);
 			ToggleLockViewportSelectionEntry.StyleNameOverride = SequencerToolbarStyleName;
 			SelectionSection.AddEntry(ToggleLockViewportSelectionEntry);
 
