[Viewport Toolbar] Fix Cinematic Viewport toolbars

- Overlay widget is now owned by SCinematicLevelViewport so that it can be shared by multiple SFilmOverlayOptions widgets
- New Toolbar placement should be fixed + general layout more in line with regular level viewport
- Each toolbar gets its own SFilmOverlayOptions button
- Old toolbar overlays widget now matching style of other buttons
- Changes name of brushes for overlays to match the name in the overlay widget

Overlays Options button location:
- Old Toolbar: right before the maximize/minimize button
- New Toolbar: within the camera menu and elevated to main toolbar level

diff --git a/Engine/Plugins/MovieScene/LevelSequenceEditor/Source/LevelSequenceEditor/Private/CinematicViewport/CinematicLevelViewportToolbarContext.h b/Engine/Plugins/MovieScene/LevelSequenceEditor/Source/LevelSequenceEditor/Private/CinematicViewport/CinematicLevelViewportToolbarContext.h
new file mode 100644
--- /dev/null
+++ b/Engine/Plugins/MovieScene/LevelSequenceEditor/Source/LevelSequenceEditor/Private/CinematicViewport/CinematicLevelViewportToolbarContext.h
@@ -0,0 +1,17 @@
+// Copyright Epic Games, Inc. All Rights Reserved.
+
+#pragma once
+
+#include "UObject/Object.h"
+#include "CinematicLevelViewportToolbarContext.generated.h"
+
+class SCinematicLevelViewport;
+
+UCLASS()
+class UCinematicLevelViewportToolbarContext : public UObject
+{
+	GENERATED_BODY()
+
+public:
+	TWeakPtr<SCinematicLevelViewport> CinematicLevelViewport;
+};
diff --git a/Engine/Plugins/MovieScene/LevelSequenceEditor/Source/LevelSequenceEditor/Private/CinematicViewport/FilmOverlays.cpp b/Engine/Plugins/MovieScene/LevelSequenceEditor/Source/LevelSequenceEditor/Private/CinematicViewport/FilmOverlays.cpp
--- a/Engine/Plugins/MovieScene/LevelSequenceEditor/Source/LevelSequenceEditor/Private/CinematicViewport/FilmOverlays.cpp
+++ b/Engine/Plugins/MovieScene/LevelSequenceEditor/Source/LevelSequenceEditor/Private/CinematicViewport/FilmOverlays.cpp
@@ -95,7 +95,7 @@ struct FFilmOverlay_Grid : IFilmOverlay
 		: NumDivsH(InNumDivsH)
 		, NumDivsV(InNumDivsV)
 	{
-		BrushName = *FString::Printf(TEXT("FilmOverlay.%dx%dGrid"), NumDivsH, NumDivsV);
+		BrushName = *FString::Printf(TEXT("FilmOverlay.Grid%dx%d"), NumDivsH, NumDivsV);
 	}
 
 	FText GetDisplayName() const
@@ -498,7 +498,7 @@ int32 SFilmOverlay::OnPaint(const FPaintArgs& Args, const FGeometry& AllottedGeo
 	return LayerId;
 }
 
-void SFilmOverlayOptions::Construct(const FArguments& InArgs)
+void SFilmOverlayOptions::Construct(const FArguments& InArgs, TSharedPtr<SFilmOverlay> InFilmOverlay)
 {
 	PrimaryColorTint = FLinearColor(1.f, 1.f, 1.f, .5f);
 
@@ -529,27 +529,61 @@ void SFilmOverlayOptions::Construct(const FArguments& InArgs)
 	ToggleableOverlays.Add(MakeShareable(new FFilmOverlay_LetterBox));
 	UFilmOverlayToolkit::RegisterToggleableFilmOverlay("LetterBox",	ToggleableOverlays.Last());
 
-	OverlayWidget = SNew(SFilmOverlay)
-		.Visibility(EVisibility::HitTestInvisible)
-		.FilmOverlays(this, &SFilmOverlayOptions::GetActiveFilmOverlays);
+	// used to choose between creating the whole button + menu widget, or just the menu widget
+	// e.g. Old toolbar uses the combo button, while the new one does not
+	const bool bIsComboButton = InArgs._IsComboButton.Get();
+
+	TSharedPtr<SWidget> MainWidget;
+	if (bIsComboButton)
+	{
+		// Styling replicates the same from the other old toolbar menu buttons
+		MainWidget =
+			SNew(SComboButton)
+			.ButtonStyle(&FAppStyle::Get().GetWidgetStyle<FButtonStyle>("EditorViewportToolBar.Button"))
+			.ForegroundColor(FSlateColor::UseStyle())
+			.HasDownArrow(false)
+			.ContentPadding(FMargin(0.0f, 0.0f))
+			.OnGetMenuContent(this, &SFilmOverlayOptions::GetMenuContent)
+			.ButtonContent()
+			[
+				SNew(SBox)
+				.VAlign(VAlign_Center)
+				.HAlign(HAlign_Center)
+				.WidthOverride(24)
+				.HeightOverride(16)
+				.ToolTipText(LOCTEXT("FilmOverlaysToolTip", "Displays a list of available film overlays to apply to this viewport."))
+				[
+					SNew(SImage)
+					.Image(this, &SFilmOverlayOptions::GetCurrentThumbnail)
+					.ColorAndOpacity(FSlateColor::UseForeground())
+				]
+			];
+	}
+	else
+	{
+		MainWidget = GetMenuContent();
+	}
+
+	OverlayWidget = InFilmOverlay;
+	if (TSharedPtr<SFilmOverlay> OverlayWidgetPinned = OverlayWidget.Pin())
+	{
+		OverlayWidgetPinned->SetFilmOverlays(TAttribute<TArray<IFilmOverlay*>>::CreateSP(this, &SFilmOverlayOptions::GetActiveFilmOverlays));
+
+		// CurrentPrimaryOverlay (and therefore Thumbnail) resets for non button version of this widget every time we create it, so let's keep it up to date
+		if (!bIsComboButton)
+		{
+			FName WidgetPrimaryOverlay = OverlayWidgetPinned->GetPrimaryFilmOverlay();
+
+			if (!WidgetPrimaryOverlay.IsNone())
+			{
+				SetPrimaryFilmOverlay(WidgetPrimaryOverlay);
+			}
+		}
+	}
 
 	ChildSlot
 	[
-		SNew(SComboButton)
-		.ButtonStyle(FAppStyle::Get(), "HoverHintOnly")
-		.ForegroundColor(FAppStyle::GetColor("InvertedForeground"))
-		.OnGetMenuContent(this, &SFilmOverlayOptions::GetMenuContent)
-		.ButtonContent()
-		[
-			SNew(SBox)
-			.WidthOverride(36)
-			.HeightOverride(24)
-			.ToolTipText(LOCTEXT("FilmOverlaysToolTip", "Displays a list of available film overlays to apply to this viewport."))
-			[
-				SNew(SImage)
-				.Image(this, &SFilmOverlayOptions::GetCurrentThumbnail)
-			]
-		]
+		MainWidget.ToSharedRef()
 	];
 }
 
@@ -753,7 +787,12 @@ TSharedRef<SWidget> SFilmOverlayOptions::ConstructToggleableOverlaysMenu()
 
 FReply SFilmOverlayOptions::SetPrimaryFilmOverlay(FName InName)
 {
-	CurrentPrimaryOverlay = InName;
+	if (TSharedPtr<SFilmOverlay> OverlayWidgetPinned = OverlayWidget.Pin())
+	{
+		OverlayWidgetPinned->SetFilmOverlays(TAttribute<TArray<IFilmOverlay*>>::CreateSP(this, &SFilmOverlayOptions::GetActiveFilmOverlays));
+		OverlayWidgetPinned->SetPrimaryFilmOverlay(InName);
+	}
+
 	IFilmOverlay* Overlay = GetPrimaryFilmOverlay();
 	if (Overlay)
 	{
@@ -773,9 +812,9 @@ FReply SFilmOverlayOptions::ToggleFilmOverlay(FName InName)
 	return FReply::Unhandled();
 }
 
-TSharedRef<SFilmOverlay> SFilmOverlayOptions::GetFilmOverlayWidget() const
+TSharedPtr<SFilmOverlay> SFilmOverlayOptions::GetFilmOverlayWidget() const
 {
-	return OverlayWidget.ToSharedRef();
+	return OverlayWidget.Pin();
 }
 
 TArray<IFilmOverlay*> SFilmOverlayOptions::GetActiveFilmOverlays() const
@@ -806,9 +845,10 @@ TArray<IFilmOverlay*> SFilmOverlayOptions::GetActiveFilmOverlays() const
 
 const FSlateBrush* SFilmOverlayOptions::GetCurrentThumbnail() const
 {
-	if (!CurrentPrimaryOverlay.IsNone())
+	FName PrimaryOverlay = GetFilmOverlayWidget()->GetPrimaryFilmOverlay();
+	if (!PrimaryOverlay.IsNone())
 	{
-		return UFilmOverlayToolkit::GetPrimaryFilmOverlays()[CurrentPrimaryOverlay].Get()->GetThumbnail();
+		return UFilmOverlayToolkit::GetPrimaryFilmOverlays()[PrimaryOverlay].Get()->GetThumbnail();
 	}
 
 	return FLevelSequenceEditorStyle::Get()->GetBrush("FilmOverlay.DefaultThumbnail");
@@ -816,9 +856,13 @@ const FSlateBrush* SFilmOverlayOptions::GetCurrentThumbnail() const
 
 IFilmOverlay* SFilmOverlayOptions::GetPrimaryFilmOverlay() const
 {
-	if (!CurrentPrimaryOverlay.IsNone())
+	if (TSharedPtr<SFilmOverlay> FilmOverlayWidget = GetFilmOverlayWidget())
 	{
-		return UFilmOverlayToolkit::GetPrimaryFilmOverlays()[CurrentPrimaryOverlay].Get();
+		FName CurrentPrimaryOverlay = FilmOverlayWidget->GetPrimaryFilmOverlay();
+		if (!CurrentPrimaryOverlay.IsNone())
+		{
+			return UFilmOverlayToolkit::GetPrimaryFilmOverlays()[CurrentPrimaryOverlay].Get();
+		}
 	}
 	return nullptr;
 }
diff --git a/Engine/Plugins/MovieScene/LevelSequenceEditor/Source/LevelSequenceEditor/Private/CinematicViewport/FilmOverlays.h b/Engine/Plugins/MovieScene/LevelSequenceEditor/Source/LevelSequenceEditor/Private/CinematicViewport/FilmOverlays.h
--- a/Engine/Plugins/MovieScene/LevelSequenceEditor/Source/LevelSequenceEditor/Private/CinematicViewport/FilmOverlays.h
+++ b/Engine/Plugins/MovieScene/LevelSequenceEditor/Source/LevelSequenceEditor/Private/CinematicViewport/FilmOverlays.h
@@ -28,23 +28,48 @@ public:
 	/** Paint this widget */
 	virtual int32 OnPaint(const FPaintArgs& Args, const FGeometry& AllottedGeometry, const FSlateRect& MyCullingRect, FSlateWindowElementList& OutDrawElements, int32 LayerId, const FWidgetStyle& InWidgetStyle, bool bParentEnabled) const override;
 
+	/** Assigns film overlays to this film overlay widget*/
+	void SetFilmOverlays(const TAttribute<TArray<IFilmOverlay*>>& InFilmOverlays)
+	{
+		FilmOverlays = InFilmOverlays;
+	}
+
+	/** Sets the current primary film overlay */
+	void SetPrimaryFilmOverlay(const FName InFilmOverlay)
+	{
+		PrimaryFilmOverlay = InFilmOverlay;
+	}
+
+	/** Get the current primary film overlay */
+	FName GetPrimaryFilmOverlay() const
+	{
+		return PrimaryFilmOverlay;
+	}
+
 private:
 	/** Attribute used once per frame to retrieve the film overlays to paint */
 	TAttribute<TArray<IFilmOverlay*>> FilmOverlays;
+
+	/** Currently selected primary film overlay */
+	FName PrimaryFilmOverlay;
 };
 
 /** A custom widget that comprises a combo box displaying all available overlay options */
 class SFilmOverlayOptions : public SCompoundWidget
 {
 public:
-	SLATE_BEGIN_ARGS(SFilmOverlayOptions){}
+	SLATE_BEGIN_ARGS(SFilmOverlayOptions)
+		: _IsComboButton(true)
+		{}
+		/** Set this to false in order to directly show the menu content, instead of showing a button to summon it */
+		SLATE_ATTRIBUTE(bool, IsComboButton)
 	SLATE_END_ARGS()
 
 	/** Construct this widget */
-	void Construct(const FArguments& InArgs);
+	void Construct(const FArguments& InArgs, TSharedPtr<SFilmOverlay> InFilmOverlay);
 
 	/** Retrieve the actual overlay widget that this widget controls. Can be positioned in any other widget hierarchy. */
-	TSharedRef<SFilmOverlay> GetFilmOverlayWidget() const;
+	TSharedPtr<SFilmOverlay> GetFilmOverlayWidget() const;
 
 	/** Bind commands for the overlays */
 	void BindCommands(TSharedRef<FUICommandList>);
@@ -82,10 +107,6 @@ private:
 	FReply ToggleFilmOverlay(FName InName);
 
 private:
-
-	/** The name of the current primary overlay */
-	FName CurrentPrimaryOverlay;
-
 	/** Color tint to apply to primary overlays */
 	FLinearColor PrimaryColorTint;
 
@@ -95,6 +116,6 @@ private:
 	/** Toggleable overlays (any number can be active at a time) */
 	TArray<TSharedPtr<IFilmOverlay>> ToggleableOverlays;
 
-	/** The overlay widget we control */
-	TSharedPtr<SFilmOverlay> OverlayWidget;
+	/** The overlay widget we control - externally owned */
+	TWeakPtr<SFilmOverlay> OverlayWidget;
 };
diff --git a/Engine/Plugins/MovieScene/LevelSequenceEditor/Source/LevelSequenceEditor/Private/CinematicViewport/SCinematicLevelViewport.cpp b/Engine/Plugins/MovieScene/LevelSequenceEditor/Source/LevelSequenceEditor/Private/CinematicViewport/SCinematicLevelViewport.cpp
--- a/Engine/Plugins/MovieScene/LevelSequenceEditor/Source/LevelSequenceEditor/Private/CinematicViewport/SCinematicLevelViewport.cpp
+++ b/Engine/Plugins/MovieScene/LevelSequenceEditor/Source/LevelSequenceEditor/Private/CinematicViewport/SCinematicLevelViewport.cpp
@@ -25,10 +25,94 @@
 #include "Widgets/Layout/SWidgetSwitcher.h"
 #include "Fonts/FontMeasure.h"
 #include "AssetEditorViewportLayout.h"
-
+#include "CinematicLevelViewportToolbarContext.h"
+#include "LevelSequenceEditorStyle.h"
+#include "Subsystems/PanelExtensionSubsystem.h"
+#include "ToolMenu.h"
+#include "ToolMenus.h"
 
 #define LOCTEXT_NAMESPACE "SCinematicLevelViewport"
 
+namespace UE::CinematicViewport::Private
+{
+void AddOverlayOptionsEntry(UToolMenu* InMenu)
+{
+	if (!InMenu)
+	{
+		return;
+	}
+
+	UCinematicLevelViewportToolbarContext* const Context = InMenu->FindContext<UCinematicLevelViewportToolbarContext>();
+	if (!Context)
+	{
+		return;
+	}
+
+	FToolMenuInsert InsertPosition("ViewportTypes", EToolMenuInsertType::After);
+	FToolMenuSection& Section = InMenu->AddSection(NAME_None, FText(), InsertPosition);
+	Section.AddSeparator(NAME_None);
+
+	TWeakPtr<SCinematicLevelViewport> CinematicLevelViewport = Context->CinematicLevelViewport;
+	TAttribute<FSlateIcon> IconAttribute = TAttribute<FSlateIcon>::CreateLambda(
+		[WeakViewport = CinematicLevelViewport]()
+		{
+			if (TSharedPtr<SCinematicLevelViewport> CinematicViewportPinned = WeakViewport.Pin())
+			{
+				if (const TSharedPtr<SFilmOverlay>& OverlayWidget = CinematicViewportPinned->GetOverlayWidget())
+				{
+					FName CurrentPrimaryOverlay = OverlayWidget->GetPrimaryFilmOverlay();
+					if (!CurrentPrimaryOverlay.IsNone())
+					{
+						// Combine strings to find the proper Icon
+						FName BrushName = FName("FilmOverlay." + CurrentPrimaryOverlay.ToString());
+						FSlateIcon Icon = FSlateIcon(FLevelSequenceEditorStyle::Get()->GetStyleSetName(), BrushName);
+						if (Icon.IsSet())
+						{
+							return Icon;
+						}
+					}
+				}
+			}
+
+			return FSlateIcon(FLevelSequenceEditorStyle::Get()->GetStyleSetName(), "FilmOverlay.Disabled");
+		}
+	);
+
+	FToolMenuEntry Entry = Section.AddSubMenu(
+		"OverlayOptions",
+		LOCTEXT("OverlayOptionsLabel", "Film Overlay Options"),
+		LOCTEXT(
+			"OverlayOptionsTooltip",
+			"Displays a list of available film overlays to apply to this viewport."
+		),
+		FNewToolMenuDelegate::CreateLambda(
+			[CinematicLevelViewport](UToolMenu* Submenu) -> void
+			{
+				if (TSharedPtr<SCinematicLevelViewport> CinematicLevelViewportPinned = CinematicLevelViewport.Pin())
+				{
+					TSharedRef<SFilmOverlayOptions> FilmOverlayOptions =
+						SNew(SFilmOverlayOptions, CinematicLevelViewportPinned->GetOverlayWidget()).IsComboButton(false);
+
+					if (TSharedPtr<FUICommandList> CommandList = CinematicLevelViewportPinned->GetCommandList())
+					{
+						FilmOverlayOptions->BindCommands(CommandList.ToSharedRef());
+					}
+
+					Submenu->AddMenuEntry(
+						"FilmOverlay", FToolMenuEntry::InitWidget("FilmOverlay", FilmOverlayOptions, FText())
+					);
+				}
+			}
+		),
+		false,
+		IconAttribute
+	);
+
+	Entry.SetShowInToolbarTopLevel(true);
+	Section.AddEntry(Entry);
+}
+} // namespace UE::CinematicViewport::Private
+
 template<typename T>
 struct SNonThrottledSpinBox : SSpinBox<T>
 {
@@ -133,15 +217,146 @@ private:
 
 class SCinematicPreviewViewport : public SLevelViewport
 {
+	static FName GetOldToolbarOverlayExtensionId()
+	{
+		return "LevelViewportToolBar.RightmostExtension";
+	}
+
 public:
+	virtual ~SCinematicPreviewViewport() override
+	{
+		// Unregistering extension from old toolbar
+		{
+			UPanelExtensionSubsystem* PanelExtensionSubsystem = GEditor->GetEditorSubsystem<UPanelExtensionSubsystem>();
+			if (!PanelExtensionSubsystem)
+			{
+				return;
+			}
+
+			if (!PanelExtensionSubsystem->IsPanelFactoryRegistered(OverlayPanelExtensionFactory.Identifier))
+			{
+				return;
+			}
+
+			PanelExtensionSubsystem->UnregisterPanelFactory(
+				OverlayPanelExtensionFactory.Identifier, GetOldToolbarOverlayExtensionId()
+			);
+		}
+	}
+
 	virtual const FSlateBrush* OnGetViewportBorderBrush() const override { return nullptr; }
 	virtual bool IsActorEditorContextVisible() const { return false; }
 	virtual EVisibility GetSelectedActorsCurrentLevelTextVisibility() const override { return EVisibility::Collapsed; }
 	virtual EVisibility GetViewportControlsVisibility() const override { return EVisibility::Collapsed; }
 
 	virtual TSharedPtr<SWidget> MakeViewportToolbar() { return nullptr; }
+	virtual TSharedPtr<SWidget> BuildViewportToolbar() override
+	{
+		return nullptr;
+	}
 
-	TSharedPtr<SWidget> MakeExternalViewportToolbar() { return SLevelViewport::MakeViewportToolbar(); }
+	// Used to add CinematicViewportContext to LevelViewport
+	virtual void ExtendToolbarContext(FToolMenuContext& InToolMenuContext) override
+	{
+		UCinematicLevelViewportToolbarContext* CinematicViewportContext =
+			NewObject<UCinematicLevelViewportToolbarContext>();
+
+		CinematicViewportContext->CinematicLevelViewport = CinematicLevelViewport;
+
+		InToolMenuContext.AddObject(CinematicViewportContext);
+	}
+
+	void SetCinematicLevelViewport(const TWeakPtr<SCinematicLevelViewport>& InCinematicLevelViewport)
+	{
+		CinematicLevelViewport = InCinematicLevelViewport;
+	}
+
+	void ExtendOldLevelToolbar()
+	{
+		if (!GEditor)
+		{
+			return;
+		}
+
+		auto CinematicLevelViewportPinned = CinematicLevelViewport.Pin();
+		if (!CinematicLevelViewportPinned)
+		{
+			return;
+		}
+
+		if (UPanelExtensionSubsystem* PanelExtensionSubsystem = GEditor->GetEditorSubsystem<UPanelExtensionSubsystem>())
+		{
+			OverlayPanelExtensionFactory.Identifier = "FilmOverlayOptions";
+			if (!PanelExtensionSubsystem->IsPanelFactoryRegistered(OverlayPanelExtensionFactory.Identifier))
+			{
+				OverlayPanelExtensionFactory.CreateExtensionWidget =
+					FPanelExtensionFactory::FCreateExtensionWidget::CreateLambda(
+						[OverlayWidgetWeak = CinematicLevelViewportPinned->GetOverlayWidget().ToWeakPtr(),
+						 CommandListWeak = CommandList.ToWeakPtr()](FWeakObjectPtr) -> TSharedRef<SWidget>
+						{
+							if (const TSharedPtr<SFilmOverlay> OverlayWidgetPinned = OverlayWidgetWeak.Pin())
+							{
+								TSharedRef<SFilmOverlayOptions> FilmOverlayOptions =
+									SNew(SFilmOverlayOptions, OverlayWidgetPinned);
+
+								if (const TSharedPtr<FUICommandList> CommandListPinned = CommandListWeak.Pin())
+								{
+									FilmOverlayOptions->BindCommands(CommandListPinned.ToSharedRef());
+								}
+
+								return FilmOverlayOptions;
+							}
+
+							return SNullWidget::NullWidget;
+						}
+					);
+
+				PanelExtensionSubsystem->RegisterPanelFactory(
+					GetOldToolbarOverlayExtensionId(), OverlayPanelExtensionFactory
+				);
+			}
+		}
+	}
+
+	static void ExtendLevelToolbar()
+	{
+		if (UToolMenu* const Submenu = UToolMenus::Get()->ExtendMenu("LevelEditor.ViewportToolbar.CameraOptions"))
+		{
+			Submenu->AddDynamicSection(
+				"CinematicCameraExtensionDynamicSection",
+				FNewToolMenuDelegate::CreateLambda(
+					[](UToolMenu* InDynamicMenu)
+					{
+						UE::CinematicViewport::Private::AddOverlayOptionsEntry(InDynamicMenu);
+					}
+				)
+			);
+		}
+	}
+
+	TSharedPtr<SWidget> MakeExternalViewportToolbar()
+	{
+		// Adding Film Overlay Options widget to both new and old toolbars
+		ExtendLevelToolbar();
+		ExtendOldLevelToolbar();
+
+		// clang-format off
+		TSharedRef<SVerticalBox> ToolbarsWidget =
+			SNew(SVerticalBox)
+			+ SVerticalBox::Slot()
+			.AutoHeight()
+			[
+				SLevelViewport::BuildViewportToolbar().ToSharedRef()
+			]
+			+ SVerticalBox::Slot()
+			.AutoHeight()
+			[
+				SLevelViewport::MakeViewportToolbar().ToSharedRef()
+			];
+		// clang-format on
+
+		return ToolbarsWidget;
+	}
 
 	FSlateColor GetBorderColorAndOpacity() const
 	{
@@ -161,8 +376,11 @@ public:
 
 private:
 	bool bShowToolbar;
-};
 
+	TWeakPtr<SCinematicLevelViewport> CinematicLevelViewport;
+
+	FPanelExtensionFactory OverlayPanelExtensionFactory;
+};
 
 BEGIN_SLATE_FUNCTION_BUILD_OPTIMIZATION
 void SCinematicLevelViewport::Construct(const FArguments& InArgs)
@@ -182,8 +400,10 @@ void SCinematicLevelViewport::Construct(const FArguments& InArgs)
 		.LevelEditorViewportClient(ViewportClient)
 		.ParentLevelEditor(InArgs._ParentLevelEditor);
 
+	ViewportWidget->SetCinematicLevelViewport(SharedThis(this));
+
 	ViewportClient->SetViewportWidget(ViewportWidget);
-	
+
 	// Automatically engage game-view to hide editor only sprites. This needs to be done
 	// after the Viewport Client and Widget are constructed as they reset the view to defaults
 	// as part of their initialization.
@@ -195,7 +415,13 @@ void SCinematicLevelViewport::Construct(const FArguments& InArgs)
 
 	FLinearColor Gray(.3f, .3f, .3f, 1.f);
 
-	TSharedRef<SFilmOverlayOptions> FilmOverlayOptions = SNew(SFilmOverlayOptions);
+	CommandList = MakeShared<FUICommandList>();
+
+	// Ensure the commands are registered
+	FCinematicViewportCommands::Register();
+	FLevelSequenceEditorCommands::Register();
+
+	OverlayWidget = SNew(SFilmOverlay).Visibility(EVisibility::HitTestInvisible);
 
 	DecoratedTransportControls = SNew(SHorizontalBox)
 		+ SHorizontalBox::Slot()
@@ -247,6 +473,7 @@ void SCinematicLevelViewport::Construct(const FArguments& InArgs)
 			SNew(SSpacer)
 		];
 
+	// clang-format off
 	TSharedRef<SWidget> MainViewport = 	SNew(SBorder)
 		.BorderImage(FAppStyle::GetBrush("BlackBrush"))
 		.ForegroundColor(Gray)
@@ -255,7 +482,6 @@ void SCinematicLevelViewport::Construct(const FArguments& InArgs)
 			SNew(SVerticalBox)
 
 			+ SVerticalBox::Slot()
-			.Padding(5.f)
 			.AutoHeight()
 			[
 				SNew(SHorizontalBox)
@@ -266,12 +492,6 @@ void SCinematicLevelViewport::Construct(const FArguments& InArgs)
 					ViewportWidget->MakeExternalViewportToolbar().ToSharedRef()
 				]
 
-				+ SHorizontalBox::Slot()
-				.AutoWidth()
-				[
-					FilmOverlayOptions					
-				]
-
 				+ SHorizontalBox::Slot()
 				.AutoWidth()
 				[
@@ -307,9 +527,7 @@ void SCinematicLevelViewport::Construct(const FArguments& InArgs)
 								ViewportWidget.ToSharedRef()
 							]
 
-							+ SOverlay::Slot()
-							[
-								FilmOverlayOptions->GetFilmOverlayWidget()
+								 + SOverlay::Slot()[OverlayWidget.ToSharedRef()
 							]
 						]
 					]
@@ -423,6 +641,7 @@ void SCinematicLevelViewport::Construct(const FArguments& InArgs)
 				]
 			]
 		];
+	// clang-format on
 
 	ChildSlot
 	[
@@ -448,14 +667,6 @@ void SCinematicLevelViewport::Construct(const FArguments& InArgs)
 		Setup(Toolkit);
 		return false;
 	});
-
-	CommandList = MakeShareable( new FUICommandList );
-
-	// Ensure the commands are registered
-	FCinematicViewportCommands::Register();
-	FLevelSequenceEditorCommands::Register();
-
-	FilmOverlayOptions->BindCommands(CommandList.ToSharedRef());
 }
 END_SLATE_FUNCTION_BUILD_OPTIMIZATION
 
diff --git a/Engine/Plugins/MovieScene/LevelSequenceEditor/Source/LevelSequenceEditor/Private/CinematicViewport/SCinematicLevelViewport.h b/Engine/Plugins/MovieScene/LevelSequenceEditor/Source/LevelSequenceEditor/Private/CinematicViewport/SCinematicLevelViewport.h
--- a/Engine/Plugins/MovieScene/LevelSequenceEditor/Source/LevelSequenceEditor/Private/CinematicViewport/SCinematicLevelViewport.h
+++ b/Engine/Plugins/MovieScene/LevelSequenceEditor/Source/LevelSequenceEditor/Private/CinematicViewport/SCinematicLevelViewport.h
@@ -4,7 +4,10 @@
 
 #include "Widgets/SCompoundWidget.h"
 #include "LevelEditorViewport.h"
+#include "Subsystems/PanelExtensionSubsystem.h"
 
+class SFilmOverlay;
+class SFilmOverlayOptions;
 namespace ETextCommit { enum Type : int; }
 struct FOptionalSize;
 
@@ -70,6 +73,8 @@ public:
 		SLATE_ARGUMENT(TSharedPtr<FAssetEditorViewportLayout>, ParentLayout)
 		/** Ptr to this viewport's parent level editor */
 		SLATE_ARGUMENT(TWeakPtr<ILevelEditor>, ParentLevelEditor)
+
+		SLATE_ARGUMENT(TWeakPtr<SCinematicPreviewViewport>, PreviewViewport)
 	SLATE_END_ARGS()
 
 	/** Access this viewport's viewport client */
@@ -83,6 +88,13 @@ public:
 	virtual void Tick(const FGeometry& AllottedGeometry, const double InCurrentTime, const float InDeltaTime);
 	virtual FReply OnKeyDown(const FGeometry& MyGeometry, const FKeyEvent& InKeyEvent) override;
 	virtual bool SupportsKeyboardFocus() const override { return true; }
+
+	const TSharedPtr<SFilmOverlay>& GetOverlayWidget()
+	{
+		return OverlayWidget;
+	}
+
+	TSharedPtr<FUICommandList> GetCommandList() { return CommandList; }
 private:
 
 	/** Set up this viewport to operate on the specified toolkit */
@@ -170,4 +182,7 @@ private:
 
 	/** The level editor viewport client for this viewport */
 	TSharedPtr<FCinematicViewportClient> ViewportClient;
+
+	/** The overlay widget */
+	TSharedPtr<SFilmOverlay> OverlayWidget;
 };
diff --git a/Engine/Plugins/MovieScene/LevelSequenceEditor/Source/LevelSequenceEditor/Private/Styles/LevelSequenceEditorStyle.h b/Engine/Plugins/MovieScene/LevelSequenceEditor/Source/LevelSequenceEditor/Private/Styles/LevelSequenceEditorStyle.h
--- a/Engine/Plugins/MovieScene/LevelSequenceEditor/Source/LevelSequenceEditor/Private/Styles/LevelSequenceEditorStyle.h
+++ b/Engine/Plugins/MovieScene/LevelSequenceEditor/Source/LevelSequenceEditor/Private/Styles/LevelSequenceEditorStyle.h
@@ -56,9 +56,9 @@ public:
 
 		Set("FilmOverlay.DefaultThumbnail", new IMAGE_BRUSH("DefaultFilmOverlayThumbnail", FVector2D(36, 24)));
 
-		Set("FilmOverlay.Disabled", new IMAGE_BRUSH("FilmOverlay.Disabled", FVector2D(36, 24)));
-		Set("FilmOverlay.2x2Grid", new IMAGE_BRUSH("FilmOverlay.2x2Grid", FVector2D(36, 24)));
-		Set("FilmOverlay.3x3Grid", new IMAGE_BRUSH("FilmOverlay.3x3Grid", FVector2D(36, 24)));
+	 	Set("FilmOverlay.Disabled", new IMAGE_BRUSH("FilmOverlay.Disabled", FVector2D(36, 24)));
+		Set("FilmOverlay.Grid2x2", new IMAGE_BRUSH("FilmOverlay.2x2Grid", FVector2D(36, 24)));
+		Set("FilmOverlay.Grid3x3", new IMAGE_BRUSH("FilmOverlay.3x3Grid", FVector2D(36, 24)));
 		Set("FilmOverlay.Crosshair", new IMAGE_BRUSH("FilmOverlay.Crosshair", FVector2D(36, 24)));
 		Set("FilmOverlay.Rabatment", new IMAGE_BRUSH("FilmOverlay.Rabatment", FVector2D(36, 24)));
 
diff --git a/Engine/Source/Editor/LevelEditor/Private/SLevelViewport.cpp b/Engine/Source/Editor/LevelEditor/Private/SLevelViewport.cpp
--- a/Engine/Source/Editor/LevelEditor/Private/SLevelViewport.cpp
+++ b/Engine/Source/Editor/LevelEditor/Private/SLevelViewport.cpp
@@ -2107,6 +2107,8 @@ TSharedPtr<SWidget> SLevelViewport::BuildViewportToolbar()
 			ContextObject->LevelViewport = SharedThis(this);
 			ViewportToolbarContext.AddObject(ContextObject);
 		}
+
+		ExtendToolbarContext(ViewportToolbarContext);
 	}
 
 	// clang-format off
diff --git a/Engine/Source/Editor/LevelEditor/Private/SLevelViewportToolBar.cpp b/Engine/Source/Editor/LevelEditor/Private/SLevelViewportToolBar.cpp
--- a/Engine/Source/Editor/LevelEditor/Private/SLevelViewportToolBar.cpp
+++ b/Engine/Source/Editor/LevelEditor/Private/SLevelViewportToolBar.cpp
@@ -276,6 +276,16 @@ void SLevelViewportToolBar::Construct( const FArguments& InArgs )
 						.Visibility(ViewportRef, &SLevelViewport::GetTransformToolbarVisibility)
 					]
 					+ SHorizontalBox::Slot()
+					.Padding(ToolbarSlotPadding)
+					.AutoWidth()
+					.HAlign(HAlign_Right)
+					.VAlign(VAlign_Fill)
+					[
+						SNew(SExtensionPanel)
+						.ExtensionPanelID("LevelViewportToolBar.RightmostExtension")
+						.ExtensionContext(ExtensionContextObject)
+					]
+					+ SHorizontalBox::Slot()
 					.HAlign(HAlign_Right)
 					.AutoWidth()
 					.Padding(ToolbarSlotPadding)
diff --git a/Engine/Source/Editor/LevelEditor/Public/SLevelViewport.h b/Engine/Source/Editor/LevelEditor/Public/SLevelViewport.h
--- a/Engine/Source/Editor/LevelEditor/Public/SLevelViewport.h
+++ b/Engine/Source/Editor/LevelEditor/Public/SLevelViewport.h
@@ -970,6 +970,9 @@ private:
 protected:
 	void LockActorInternal(AActor* NewActorToLock);
 
+	/** Can be overriden by derived classes to add new context objects to the new toolbar */
+	virtual void ExtendToolbarContext(FToolMenuContext& InToolMenuContext) {};
+
 public:
 	static bool GetCameraInformationFromActor(AActor* Actor, FMinimalViewInfo& out_CameraInfo);
 
