Sequencer: Fix issue where properties on spawnables were being persisted inconsistently.
Spawn code can get called and may invalidate the object tracking 'modified' state right before an object is despawned when interacting with certain UI elements (e.g. the mute). Now we track to see whether tracked object state already exists for an object being creating new tracked object state.

diff --git a/Engine/Plugins/MovieScene/LevelSequenceEditor/Source/LevelSequenceEditor/Private/Misc/LevelSequenceEditorSpawnRegister.cpp b/Engine/Plugins/MovieScene/LevelSequenceEditor/Source/LevelSequenceEditor/Private/Misc/LevelSequenceEditorSpawnRegister.cpp
--- a/Engine/Plugins/MovieScene/LevelSequenceEditor/Source/LevelSequenceEditor/Private/Misc/LevelSequenceEditorSpawnRegister.cpp
+++ b/Engine/Plugins/MovieScene/LevelSequenceEditor/Source/LevelSequenceEditor/Private/Misc/LevelSequenceEditorSpawnRegister.cpp
@@ -81,8 +81,11 @@ UObject* FLevelSequenceEditorSpawnRegister::SpawnObject(const FGuid& BindingId,
 		// Mark as replay rewindable so it persists while a replay is going on in PIE.
 		NewActor->bReplayRewindable = true;
 
-		// Add an entry to the tracked objects map to keep track of this object (so that it can be saved when modified)
-		TrackedObjects.Add(NewActor, FTrackedObjectState(Template, BindingId, BindingIndex));
+		if (!TrackedObjects.Find(NewActor))
+		{ 
+			// Add an entry to the tracked objects map to keep track of this object (so that it can be saved when modified)
+			TrackedObjects.Add(NewActor, FTrackedObjectState(Template, BindingId, BindingIndex));
+		}
 
 		if (ULayersSubsystem* Layers = GEditor->GetEditorSubsystem<ULayersSubsystem>())
 		{
