Sequencer: Added ULevelSequenceEditorSubsystem::AddSpawnableFromInstance and AddSpawnableFromClass

Spawnables needs to access the actor factories. This addresses a regression in 5.5.

Deprecated UMovieSceneSequenceExtensions::AddSpawnableFromInstance and AddSpawnableFromClass

diff --git a/Engine/Plugins/MovieScene/LevelSequenceEditor/Source/LevelSequenceEditor/Private/LevelSequenceEditorSubsystem.cpp b/Engine/Plugins/MovieScene/LevelSequenceEditor/Source/LevelSequenceEditor/Private/LevelSequenceEditorSubsystem.cpp
--- a/Engine/Plugins/MovieScene/LevelSequenceEditor/Source/LevelSequenceEditor/Private/LevelSequenceEditorSubsystem.cpp
+++ b/Engine/Plugins/MovieScene/LevelSequenceEditor/Source/LevelSequenceEditor/Private/LevelSequenceEditorSubsystem.cpp
@@ -841,6 +841,52 @@ TArray<FMovieSceneBindingProxy> ULevelSequenceEditorSubsystem::AddActors(const T
 	return BindingProxies;
 }
 
+FMovieSceneBindingProxy ULevelSequenceEditorSubsystem::AddSpawnableFromInstance(UObject* ObjectToSpawn)
+{
+	TSharedPtr<ISequencer> Sequencer = GetActiveSequencer();
+	if (Sequencer == nullptr)
+	{
+		return FMovieSceneBindingProxy();
+	}
+
+	UMovieSceneSequence* Sequence = Sequencer->GetFocusedMovieSceneSequence();
+	if (!Sequence)
+	{
+		return FMovieSceneBindingProxy();
+	}
+
+	UE::Sequencer::FCreateBindingParams Params;
+	Params.BindingNameOverride = ObjectToSpawn->GetName();
+	Params.bSpawnable = true;
+
+	FGuid Guid = FSequencerUtilities::CreateBinding(Sequencer.ToSharedRef(), *ObjectToSpawn, Params);
+	FMovieSceneBindingProxy BindingProxy(Guid, Sequence);
+	return BindingProxy;
+}
+
+FMovieSceneBindingProxy ULevelSequenceEditorSubsystem::AddSpawnableFromClass(UClass* ClassToSpawn)
+{
+	TSharedPtr<ISequencer> Sequencer = GetActiveSequencer();
+	if (Sequencer == nullptr)
+	{
+		return FMovieSceneBindingProxy();
+	}
+
+	UMovieSceneSequence* Sequence = Sequencer->GetFocusedMovieSceneSequence();
+	if (!Sequence)
+	{
+		return FMovieSceneBindingProxy();
+	}
+
+	UE::Sequencer::FCreateBindingParams Params;
+	Params.BindingNameOverride = ClassToSpawn->GetName();
+	Params.bSpawnable = true;
+
+	FGuid Guid = FSequencerUtilities::CreateBinding(Sequencer.ToSharedRef(), *ClassToSpawn, Params);
+	FMovieSceneBindingProxy BindingProxy(Guid, Sequence);
+	return BindingProxy;
+}
+
 FMovieSceneBindingProxy ULevelSequenceEditorSubsystem::CreateCamera(bool bSpawnable, ACineCameraActor*& OutActor)
 {
 	TSharedPtr<ISequencer> Sequencer = GetActiveSequencer();
diff --git a/Engine/Plugins/MovieScene/LevelSequenceEditor/Source/LevelSequenceEditor/Public/LevelSequenceEditorSubsystem.h b/Engine/Plugins/MovieScene/LevelSequenceEditor/Source/LevelSequenceEditor/Public/LevelSequenceEditorSubsystem.h
--- a/Engine/Plugins/MovieScene/LevelSequenceEditor/Source/LevelSequenceEditor/Public/LevelSequenceEditorSubsystem.h
+++ b/Engine/Plugins/MovieScene/LevelSequenceEditor/Source/LevelSequenceEditor/Public/LevelSequenceEditorSubsystem.h
@@ -112,6 +112,14 @@ public:
 	UFUNCTION(BlueprintCallable, Category = "Level Sequence Editor")
 	TArray<FMovieSceneBindingProxy> AddActors(const TArray<AActor*>& Actors);
 
+	/** Add a new binding to this sequence that will spawn the specified object. */
+	UFUNCTION(BlueprintCallable, Category = "Level Sequence Editor")
+	FMovieSceneBindingProxy AddSpawnableFromInstance(UObject* ObjectToSpawn);
+
+	/** Add a new binding to this sequence that will spawn the specified class. */
+	UFUNCTION(BlueprintCallable, Category = "Level Sequence Editor")
+	FMovieSceneBindingProxy AddSpawnableFromClass(UClass* ClassToSpawn);
+
 	/** Create a cine camera actor and add it to Sequencer */
 	UFUNCTION(BlueprintCallable, Category = "Level Sequence Editor")
 	FMovieSceneBindingProxy CreateCamera(bool bSpawnable, ACineCameraActor*& OutActor);
diff --git a/Engine/Plugins/MovieScene/SequencerScripting/Source/SequencerScripting/Public/ExtensionLibraries/MovieSceneSequenceExtensions.h b/Engine/Plugins/MovieScene/SequencerScripting/Source/SequencerScripting/Public/ExtensionLibraries/MovieSceneSequenceExtensions.h
--- a/Engine/Plugins/MovieScene/SequencerScripting/Source/SequencerScripting/Public/ExtensionLibraries/MovieSceneSequenceExtensions.h
+++ b/Engine/Plugins/MovieScene/SequencerScripting/Source/SequencerScripting/Public/ExtensionLibraries/MovieSceneSequenceExtensions.h
@@ -407,7 +407,8 @@ public:
 	 * @param ObjectToSpawn   An object instance to use as a template for spawning
 	 * @return A unique identifier for the new binding
 	 */
-	UFUNCTION(BlueprintCallable, Category="Sequencer|Sequence", meta=(ScriptMethod))
+	UE_DEPRECATED(5.5, "AddSpawnableFromInstance is deprecated. Please use AddSpawnableFromInstance in the Level Sequence Editor Subsystem")
+	UFUNCTION(BlueprintCallable, Category="Sequencer|Sequence", meta=(ScriptMethod, DeprecatedFunction, DeprecationMessage = "AddSpawnableFromInstance is deprecated. Please use AddSpawnableFromInstance in the Level Sequence Editor Subsystem"))
 	static FMovieSceneBindingProxy AddSpawnableFromInstance(UMovieSceneSequence* Sequence, UObject* ObjectToSpawn);
 
 	/**
@@ -418,7 +419,8 @@ public:
 	 * @param ClassToSpawn    A class or blueprint type to spawn for this binding
 	 * @return A unique identifier for the new binding
 	 */
-	UFUNCTION(BlueprintCallable, Category="Sequencer|Sequence", meta=(ScriptMethod))
+	UE_DEPRECATED(5.5, "AddSpawnableFromClass is deprecated. Please use AddSpawnableFromClass in the Level Sequence Editor Subsystem")
+	UFUNCTION(BlueprintCallable, Category="Sequencer|Sequence", meta=(ScriptMethod, DeprecatedFunction, DeprecationMessage = "AddSpawnableFromClass is deprecated. Please use AddSpawnableFromClass in the Level Sequence Editor Subsystem"))
 	static FMovieSceneBindingProxy AddSpawnableFromClass(UMovieSceneSequence* Sequence, UClass* ClassToSpawn);
 
 	/**
