Sequencer: Fix RefreshBindingDetails so that it doesn't create a new BindingPropertyInfoList when the menu is still visible and instead just refills the existing menu. This fixes an issue where if you right click on Binding Properties and modify UOLs, it would only complete the first operation. For example, you have to open the menu, set the fragment to Actor, close the menu, open the menu, pick the Actor (you couldn't do this all with the menu open once).

diff --git a/Engine/Plugins/MovieScene/LevelSequenceEditor/Source/LevelSequenceEditor/Private/LevelSequenceEditorSubsystem.cpp b/Engine/Plugins/MovieScene/LevelSequenceEditor/Source/LevelSequenceEditor/Private/LevelSequenceEditorSubsystem.cpp
--- a/Engine/Plugins/MovieScene/LevelSequenceEditor/Source/LevelSequenceEditor/Private/LevelSequenceEditorSubsystem.cpp
+++ b/Engine/Plugins/MovieScene/LevelSequenceEditor/Source/LevelSequenceEditor/Private/LevelSequenceEditorSubsystem.cpp
@@ -3115,7 +3115,15 @@ void ULevelSequenceEditorSubsystem::RefreshBindingDetails(IDetailsView* DetailsV
 	}
 	if (FMovieSceneBindingReferences* BindingReferences = Sequence->GetBindingReferences())
 	{
-		BindingPropertyInfoList = NewObject<UMovieSceneBindingPropertyInfoList>(this);
+		if (!BindingPropertyInfoList)
+		{
+			BindingPropertyInfoList = NewObject<UMovieSceneBindingPropertyInfoList>(this);
+		}
+		else
+		{
+			BindingPropertyInfoList->Bindings.Empty();
+		}
+
 		Algo::Transform(BindingReferences->GetReferences(ObjectBindingID), BindingPropertyInfoList->Bindings, [this](const FMovieSceneBindingReference& Reference)
 			{
 				UMovieSceneCustomBinding* CopiedBinding = nullptr;
