[Sequencer] Navigation Tool : disable initialization for now as there are side effects

#rnx

diff --git a/Engine/Plugins/MovieScene/LevelSequenceEditor/Source/LevelSequenceEditor/Private/LevelSequenceEditorToolkit.cpp b/Engine/Plugins/MovieScene/LevelSequenceEditor/Source/LevelSequenceEditor/Private/LevelSequenceEditorToolkit.cpp
--- a/Engine/Plugins/MovieScene/LevelSequenceEditor/Source/LevelSequenceEditor/Private/LevelSequenceEditorToolkit.cpp
+++ b/Engine/Plugins/MovieScene/LevelSequenceEditor/Source/LevelSequenceEditor/Private/LevelSequenceEditorToolkit.cpp
@@ -319,12 +319,12 @@ void FLevelSequenceEditorToolkit::Initialize(const EToolkitMode::Type Mode, cons
 		}
 	}
 
-	using namespace UE::Sequencer::NavigationTool;
-	if (const TSharedPtr<INavigationTool> Tool = Sequencer->GetNavigationTool())
-	{
-		NavigationToolProvider = MakeShared<FLevelSequenceNavigationToolProvider>();
-		FNavigationToolExtender::RegisterOnlyToolProvider(Tool.ToSharedRef(), NavigationToolProvider.ToSharedRef());
-	}
+	// using namespace UE::Sequencer::NavigationTool;
+	// if (const TSharedPtr<INavigationTool> Tool = Sequencer->GetNavigationTool())
+	// {
+	// 	NavigationToolProvider = MakeShared<FLevelSequenceNavigationToolProvider>();
+	// 	FNavigationToolExtender::RegisterOnlyToolProvider(Tool.ToSharedRef(), NavigationToolProvider.ToSharedRef());
+	// }
 }
 
 
