Sequencer: Genericized 'Spawnable Object Binding ID' to work for any binding type, renamed to 'Proxy'

I will fix up the internal API around these in future, but for now I'm just adding the ability to easily change a binding to be a proxy

diff --git a/Engine/Plugins/MovieScene/LevelSequenceEditor/Source/LevelSequenceEditor/Private/LevelSequenceEditorSubsystem.cpp b/Engine/Plugins/MovieScene/LevelSequenceEditor/Source/LevelSequenceEditor/Private/LevelSequenceEditorSubsystem.cpp
--- a/Engine/Plugins/MovieScene/LevelSequenceEditor/Source/LevelSequenceEditor/Private/LevelSequenceEditorSubsystem.cpp
+++ b/Engine/Plugins/MovieScene/LevelSequenceEditor/Source/LevelSequenceEditor/Private/LevelSequenceEditorSubsystem.cpp
@@ -91,6 +91,63 @@ DEFINE_LOG_CATEGORY(LogLevelSequenceEditor);
 
 #define LOCTEXT_NAMESPACE "LevelSequenceEditor"
 
+namespace UE::Sequencer
+{
+	struct FProxyObjectBindingIDPicker : FMovieSceneObjectBindingIDPicker
+	{
+		using FMovieSceneObjectBindingIDPicker::GetPickerMenu;
+		using FMovieSceneObjectBindingIDPicker::GetCurrentItemWidget;
+
+		FProxyObjectBindingIDPicker(TSharedPtr<ISequencer> InSequencer, const FGuid& InObjectBindingID, TFunction<void(FMovieSceneObjectBindingID)> InPreOnPicked = nullptr)
+			: FMovieSceneObjectBindingIDPicker(InSequencer->GetFocusedTemplateID(), InSequencer)
+			, ObjectBindingID(InObjectBindingID)
+			, PreOnPicked(InPreOnPicked)
+		{
+			Initialize();
+		}
+
+		virtual UMovieSceneSequence* GetSequence() const override
+		{
+			return WeakSequencer.Pin() ? WeakSequencer.Pin()->GetFocusedMovieSceneSequence() : nullptr;
+		}
+		virtual void SetCurrentValue(const FMovieSceneObjectBindingID& InBindingId) override
+		{
+			FScopedTransaction Transaction(LOCTEXT("ChangeBindingProxyTransaction", "Change Proxy Binding"));
+
+			if (UMovieSceneSequence* Sequence = GetSequence())
+			{
+				if (PreOnPicked)
+				{
+					PreOnPicked(InBindingId);
+				}
+
+				FMovieScenePossessable* Possessable = Sequence->GetMovieScene()->FindPossessable(ObjectBindingID);
+				if (Possessable)
+				{
+					Possessable->SetSpawnableObjectBindingID(InBindingId);
+				}
+			}
+		}
+
+		virtual FMovieSceneObjectBindingID GetCurrentValue() const override
+		{
+			if (UMovieSceneSequence* Sequence = GetSequence())
+			{
+				FMovieScenePossessable* Possessable = Sequence->GetMovieScene()->FindPossessable(ObjectBindingID);
+				if (Possessable)
+				{
+					return Possessable->GetSpawnableObjectBindingID();
+				}
+			}
+			return FMovieSceneObjectBindingID();
+		}
+
+	private:
+		FGuid ObjectBindingID;
+		TFunction<void(FMovieSceneObjectBindingID)> PreOnPicked;
+	};
+}
+
 class FMovieSceneBindingPropertyInfoListCustomization : public IDetailCustomization
 {
 public:
@@ -340,6 +397,7 @@ private:
 	int32 BindingIndex = 0;
 
 	TSharedPtr<IPropertyHandle> StructPropertyHandle;
+	TSharedPtr<UE::Sequencer::FProxyObjectBindingIDPicker> ProxyPicker;
 };
 
 FMovieSceneBindingPropertyInfoDetailCustomization::FMovieSceneBindingPropertyInfoDetailCustomization(TWeakPtr<ISequencer> InSequencer, UMovieScene* InMovieScene, FGuid InBindingGuid, ULevelSequenceEditorSubsystem* InLevelSequenceEditorSubsystem)
@@ -416,11 +474,40 @@ void FMovieSceneBindingPropertyInfoDetailCustomization::CustomizeChildren(TShare
 
 			if (!BindingPropertyInfo->CustomBinding)
 			{
-				// Show locator property
-				TSharedPtr<IPropertyHandle> LocatorProperty = StructPropertyHandle->GetChildHandle(GET_MEMBER_NAME_CHECKED(FMovieSceneBindingPropertyInfo, Locator));
-				if (LocatorProperty.IsValid())
+				FMovieScenePossessable* Possessable = Sequence->GetMovieScene()->FindPossessable(BindingGuid);
+				if (Possessable && Possessable->GetSpawnableObjectBindingID().IsValid())
 				{
-					StructBuilder.AddProperty(LocatorProperty.ToSharedRef());
+					ProxyPicker = MakeShared<UE::Sequencer::FProxyObjectBindingIDPicker>(Sequencer, this->BindingGuid);
+
+					StructBuilder.AddCustomRow(FText())
+					.NameContent()
+					[
+						SNew(STextBlock)
+						.Text(LOCTEXT("ProxyLabel", "Proxy Binding"))
+						.Font(IDetailLayoutBuilder::GetDetailFont())
+					]
+					.ValueContent()
+					[
+						SNew(SComboButton)
+						.OnGetMenuContent(ProxyPicker.ToSharedRef(), &UE::Sequencer::FProxyObjectBindingIDPicker::GetPickerMenu)
+						.ContentPadding(FMargin(4.0, 2.0))
+						.ButtonContent()
+						[
+							ProxyPicker->GetCurrentItemWidget(
+								SNew(STextBlock)
+								.Font(CustomizationUtils.GetRegularFont())
+							)
+						]
+					];
+				}
+				else
+				{
+					// Show locator property
+					TSharedPtr<IPropertyHandle> LocatorProperty = StructPropertyHandle->GetChildHandle(GET_MEMBER_NAME_CHECKED(FMovieSceneBindingPropertyInfo, Locator));
+					if (LocatorProperty.IsValid())
+					{
+						StructBuilder.AddProperty(LocatorProperty.ToSharedRef());
+					}
 				}
 			}
 			else
@@ -3247,6 +3334,8 @@ void ULevelSequenceEditorSubsystem::RefreshTrackRowMetadataDetails(IDetailsView*
 
 void ULevelSequenceEditorSubsystem::AddChangeBindingTypeMenu(FMenuBuilder& MenuBuilder, TSharedRef<ISequencer> Sequencer, const TArray<FSequencerChangeBindingInfo>& BindingsToChange, bool bConvert, TFunction<void()> OnBindingChanged)
 {
+	using namespace UE::Sequencer;
+
 	UMovieSceneSequence* Sequence = Sequencer->GetFocusedMovieSceneSequence();
 	UMovieScene* MovieScene = Sequence ? Sequence->GetMovieScene() : nullptr;
 
@@ -3255,13 +3344,48 @@ void ULevelSequenceEditorSubsystem::AddChangeBindingTypeMenu(FMenuBuilder& MenuB
 		return;
 	}
 
+	const TWeakPtr<ISequencer> WeakSequencer = Sequencer;
+
+	if (BindingsToChange.Num() == 1)
+	{
+		// This is captured by-value into the submenu lambda to keep it alive while the menu is open
+		TSharedPtr<FProxyObjectBindingIDPicker> ProxyPicker = MakeShared<FProxyObjectBindingIDPicker>(Sequencer, BindingsToChange[0].BindingID, [this, WeakSequencer, Sequence, BindingsToChange, OnBindingChanged](FMovieSceneObjectBindingID ID){
+
+			if (TSharedPtr<ISequencer> SequencerPtr = WeakSequencer.Pin())
+			{
+				auto OnChanged = [Sequence, ID](FGuid BindingID, int32 BindingIndex)
+				{
+					if (FMovieSceneBindingReferences* BindingReferences = Sequence->GetBindingReferences())
+					{
+						BindingReferences->AddOrReplaceBinding(BindingID, FUniversalObjectLocator(), BindingIndex);
+					}
+
+					return Sequence->GetMovieScene()->FindPossessable(BindingID);
+				};
+
+				this->ChangeBindingTypes(
+					SequencerPtr.ToSharedRef(),
+					BindingsToChange,
+					OnChanged,
+					OnBindingChanged
+				);
+			}
+		});
+
+		MenuBuilder.AddSubMenu(
+			LOCTEXT("ConvertToProxy", "Convert to Proxy"),
+			LOCTEXT("ConvertToProxyTooltip", "Convert selected binding(s) to a proxy binding that simply references a binding in another sequence"),
+			FNewMenuDelegate::CreateLambda([ProxyPicker](FMenuBuilder& SubMenuBuilder) {
+				ProxyPicker->GetPickerMenu(SubMenuBuilder);
+			})
+		);
+	}
+
 	// Can convert to possessable
 	if (!bConvert || Algo::AllOf(BindingsToChange, [this, Sequencer, Sequence, MovieScene](const FSequencerChangeBindingInfo& BindingInfo) {
 		return FSequencerUtilities::CanConvertToPossessable(Sequencer, BindingInfo.BindingID, BindingInfo.BindingIndex);
 		}))
 	{
-		const TWeakPtr<ISequencer> WeakSequencer = Sequencer;
-
 		MenuBuilder.AddMenuEntry(
 			LOCTEXT("ConvertToPossessable", "Possessable"),
 			bConvert ? LOCTEXT("ConvertToPossessableTooltip", "Convert selected binding(s) to a possessable") : LOCTEXT("ChangeToPossessableTooltip", "Reset selected binding(s) to a new possessable"),
@@ -3324,8 +3448,6 @@ void ULevelSequenceEditorSubsystem::AddChangeBindingTypeMenu(FMenuBuilder& MenuB
 				return FSequencerUtilities::CanConvertToCustomBinding(Sequencer, BindingInfo.BindingID, CustomBindingType, BindingInfo.BindingIndex);
 			}))
 		{
-			const TWeakPtr<ISequencer> WeakSequencer = Sequencer;
-
 			// Special case director blueprint bindings to show sub-menu for setting endpoints
 			if (CustomBindingType == UMovieSceneSpawnableDirectorBlueprintBinding::StaticClass()
 				|| CustomBindingType == UMovieSceneReplaceableDirectorBlueprintBinding::StaticClass())
diff --git a/Engine/Source/Editor/EditorStyle/Private/StarshipStyle.cpp b/Engine/Source/Editor/EditorStyle/Private/StarshipStyle.cpp
--- a/Engine/Source/Editor/EditorStyle/Private/StarshipStyle.cpp
+++ b/Engine/Source/Editor/EditorStyle/Private/StarshipStyle.cpp
@@ -2368,6 +2368,7 @@ void FStarshipEditorStyle::FStyle::SetupSequencerStyles()
 		Set("Sequencer.SpawnableIconOverlay", new IMAGE_BRUSH(TEXT("Sequencer/SpawnableIconOverlay"), FVector2f(13.f, 13.f)));
 		Set("Sequencer.ReplaceableIconOverlay", new IMAGE_BRUSH(TEXT("Sequencer/ReplaceableIconOverlay"), FVector2f(13.f, 13.f)));
 		Set("Sequencer.MultipleIconOverlay", new IMAGE_BRUSH(TEXT("Sequencer/MultipleIconOverlay"), FVector2f(13.f, 13.f)));
+		Set("Sequencer.ProxyIconOverlay", new IMAGE_BRUSH(TEXT("Sequencer/ProxyIconOverlay"), FVector2f(13.f, 13.f)));
 		Set("Sequencer.DynamicBindingIconOverlay", new IMAGE_BRUSH(TEXT("Sequencer/DynamicBindingIconOverlay"), Icon16x16));
 		Set("Sequencer.SpawnableDynamicBindingIconOverlay", new IMAGE_BRUSH(TEXT("Sequencer/SpawnableDynamicBindingIconOverlay"), Icon16x16));
 		Set("Sequencer.LockSequence", new IMAGE_BRUSH("Sequencer/Main_Icons/Icon_Sequencer_Locked_16x", Icon16x16));
diff --git a/Engine/Source/Editor/MovieSceneTools/Private/MovieSceneObjectBindingIDPicker.cpp b/Engine/Source/Editor/MovieSceneTools/Private/MovieSceneObjectBindingIDPicker.cpp
--- a/Engine/Source/Editor/MovieSceneTools/Private/MovieSceneObjectBindingIDPicker.cpp
+++ b/Engine/Source/Editor/MovieSceneTools/Private/MovieSceneObjectBindingIDPicker.cpp
@@ -224,7 +224,7 @@ bool FMovieSceneObjectBindingIDPicker::IsClassAllowed(const TSharedPtr<FSequence
 
 	InClass = MovieSceneHelpers::GetBoundObjectClass(BindSequence, BindingID.Guid);
 
-	if (InClass == nullptr)
+	if (InClass == nullptr || InClass == UObject::StaticClass())
 	{
 		return true;
 	}
diff --git a/Engine/Source/Editor/Sequencer/Private/MVVM/ViewModels/PossessableModel.cpp b/Engine/Source/Editor/Sequencer/Private/MVVM/ViewModels/PossessableModel.cpp
--- a/Engine/Source/Editor/Sequencer/Private/MVVM/ViewModels/PossessableModel.cpp
+++ b/Engine/Source/Editor/Sequencer/Private/MVVM/ViewModels/PossessableModel.cpp
@@ -135,6 +135,12 @@ const FSlateBrush* FPossessableModel::GetIconOverlayBrush() const
 	UMovieSceneSequence* Sequence = OwnerModel ? OwnerModel->GetSequence() : nullptr;
 	if (Sequence)
 	{
+		const FMovieScenePossessable* Possessable = Sequence->GetMovieScene()->FindPossessable(ObjectBindingID);
+		if (Possessable && Possessable->GetSpawnableObjectBindingID().IsValid())
+		{
+			return FAppStyle::GetBrush("Sequencer.ProxyIconOverlay");
+		}
+
 		if (const FMovieSceneBindingReferences* BindingReferences = Sequence->GetBindingReferences())
 		{
 			for (const FMovieSceneBindingReference& BindingReference : BindingReferences->GetReferences(ObjectBindingID))
@@ -164,7 +170,37 @@ const UClass* FPossessableModel::FindObjectClass() const
 			FMovieScenePossessable* Possessable = MovieScene->FindPossessable(ObjectBindingID);
 			if (TSharedPtr<ISequencer> Sequencer = OwnerModel->GetSequencer())
 			{
-				if (UObject* Template = MovieSceneHelpers::GetObjectTemplate(OwnerModel->GetSequence(), ObjectBindingID, Sequencer->GetSharedPlaybackState()))
+				UMovieSceneSequence* Sequence = OwnerModel->GetSequence();
+				FGuid BindingID = ObjectBindingID;
+				FMovieSceneSequenceID SequenceID = Sequencer->GetFocusedTemplateID();
+
+				using VisitedKeyType = TTuple<UMovieSceneSequence*, FGuid>;
+				TSet<VisitedKeyType, DefaultKeyFuncs<VisitedKeyType>, TInlineSetAllocator<8>> VisitedBindings;
+				VisitedBindings.Add(MakeTuple(Sequence, BindingID));
+
+				// Find the tar\get for the proxy by looping through proxies
+				while (Possessable && Possessable->GetSpawnableObjectBindingID().IsValid())
+				{
+					SequenceID = Possessable->GetSpawnableObjectBindingID().ResolveSequenceID(SequenceID, *Sequencer);
+					BindingID  = Possessable->GetSpawnableObjectBindingID().GetGuid();
+					Sequence   = Sequencer->GetSharedPlaybackState()->GetSequence(SequenceID);
+
+					if (!Sequence)
+					{
+						// Broken link!
+						return Class;
+					}
+					if (VisitedBindings.Contains(MakeTuple(Sequence, BindingID)))
+					{
+						// Cyclic dependency!
+						return nullptr;
+					}
+
+					VisitedBindings.Add(MakeTuple(Sequence, BindingID));
+					Possessable = Sequence->GetMovieScene()->FindPossessable(BindingID);
+				}
+
+				if (UObject* Template = MovieSceneHelpers::GetObjectTemplate(Sequence, BindingID, Sequencer->GetSharedPlaybackState()))
 				{
 					Class = Template->GetClass();
 				}
