Fix issue where the legacy cinematic options button showed up in the new viewport toolbar.

diff --git a/Engine/Plugins/MovieScene/LevelSequenceEditor/Source/LevelSequenceEditor/Private/CinematicViewport/SCinematicLevelViewport.cpp b/Engine/Plugins/MovieScene/LevelSequenceEditor/Source/LevelSequenceEditor/Private/CinematicViewport/SCinematicLevelViewport.cpp
--- a/Engine/Plugins/MovieScene/LevelSequenceEditor/Source/LevelSequenceEditor/Private/CinematicViewport/SCinematicLevelViewport.cpp
+++ b/Engine/Plugins/MovieScene/LevelSequenceEditor/Source/LevelSequenceEditor/Private/CinematicViewport/SCinematicLevelViewport.cpp
@@ -299,6 +299,11 @@ public:
 							{
 								TSharedRef<SFilmOverlayOptions> FilmOverlayOptions =
 									SNew(SFilmOverlayOptions, OverlayWidgetPinned);
+									
+								FilmOverlayOptions->SetVisibility(TAttribute<EVisibility>::CreateLambda([]
+								{
+									return UE::UnrealEd::ShowOldViewportToolbars() ? EVisibility::Visible : EVisibility::Collapsed;
+								}));
 
 								if (const TSharedPtr<FUICommandList> CommandListPinned = CommandListWeak.Pin())
 								{
