Add small film overlay icons for toolbar menus.

#rnx

diff --git a/Engine/Plugins/MovieScene/LevelSequenceEditor/Source/LevelSequenceEditor/Private/CinematicViewport/SCinematicLevelViewport.cpp b/Engine/Plugins/MovieScene/LevelSequenceEditor/Source/LevelSequenceEditor/Private/CinematicViewport/SCinematicLevelViewport.cpp
--- a/Engine/Plugins/MovieScene/LevelSequenceEditor/Source/LevelSequenceEditor/Private/CinematicViewport/SCinematicLevelViewport.cpp
+++ b/Engine/Plugins/MovieScene/LevelSequenceEditor/Source/LevelSequenceEditor/Private/CinematicViewport/SCinematicLevelViewport.cpp
@@ -65,7 +65,7 @@ void AddOverlayOptionsEntry(UToolMenu* InMenu)
 					if (!CurrentPrimaryOverlay.IsNone())
 					{
 						// Combine strings to find the proper Icon
-						FName BrushName = FName("FilmOverlay." + CurrentPrimaryOverlay.ToString());
+						FName BrushName = FName("FilmOverlay." + CurrentPrimaryOverlay.ToString() + ".Small");
 						FSlateIcon Icon = FSlateIcon(FLevelSequenceEditorStyle::Get()->GetStyleSetName(), BrushName);
 						if (Icon.IsSet())
 						{
@@ -75,7 +75,7 @@ void AddOverlayOptionsEntry(UToolMenu* InMenu)
 				}
 			}
 
-			return FSlateIcon(FLevelSequenceEditorStyle::Get()->GetStyleSetName(), "FilmOverlay.Disabled");
+			return FSlateIcon(FLevelSequenceEditorStyle::Get()->GetStyleSetName(), "FilmOverlay.Disabled.Small");
 		}
 	);
 
diff --git a/Engine/Plugins/MovieScene/LevelSequenceEditor/Source/LevelSequenceEditor/Private/Styles/LevelSequenceEditorStyle.h b/Engine/Plugins/MovieScene/LevelSequenceEditor/Source/LevelSequenceEditor/Private/Styles/LevelSequenceEditorStyle.h
--- a/Engine/Plugins/MovieScene/LevelSequenceEditor/Source/LevelSequenceEditor/Private/Styles/LevelSequenceEditorStyle.h
+++ b/Engine/Plugins/MovieScene/LevelSequenceEditor/Source/LevelSequenceEditor/Private/Styles/LevelSequenceEditorStyle.h
@@ -61,6 +61,12 @@ public:
 		Set("FilmOverlay.Grid3x3", new IMAGE_BRUSH("FilmOverlay.3x3Grid", FVector2D(36, 24)));
 		Set("FilmOverlay.Crosshair", new IMAGE_BRUSH("FilmOverlay.Crosshair", FVector2D(36, 24)));
 		Set("FilmOverlay.Rabatment", new IMAGE_BRUSH("FilmOverlay.Rabatment", FVector2D(36, 24)));
+		
+	 	Set("FilmOverlay.Disabled.Small", new IMAGE_BRUSH_SVG("FilmOverlay.Disabled_16", Icon16x16));
+	 	Set("FilmOverlay.Grid2x2.Small", new IMAGE_BRUSH_SVG("FilmOverlay.2x2Grid_16", Icon16x16));
+	 	Set("FilmOverlay.Grid3x3.Small", new IMAGE_BRUSH_SVG("FilmOverlay.3x3Grid_16", Icon16x16));
+	 	Set("FilmOverlay.Crosshair.Small", new IMAGE_BRUSH_SVG("FilmOverlay.Crosshair_16", Icon16x16));
+	 	Set("FilmOverlay.Rabatment.Small", new IMAGE_BRUSH_SVG("FilmOverlay.Rabatment_16", Icon16x16));
 
 		FSlateStyleRegistry::RegisterSlateStyle(*this);
 	 }
