[Sequencer] Move Navigation Tool to new experimental SequenceNavigator plugin

diff --git a/Engine/Plugins/Experimental/SequenceNavigator/Source/SequenceNavigator/Public/NavigationToolViewModelTypeID.h b/Engine/Plugins/Experimental/SequenceNavigator/Source/SequenceNavigator/Public/NavigationToolViewModelTypeID.h
new file mode 100644
--- /dev/null
+++ b/Engine/Plugins/Experimental/SequenceNavigator/Source/SequenceNavigator/Public/NavigationToolViewModelTypeID.h
@@ -0,0 +1,109 @@
+﻿// Copyright Epic Games, Inc. All Rights Reserved.
+
+#pragma once
+
+#include "CoreTypes.h"
+#include "NavigationToolCastableTypeTable.h"
+
+namespace UE::Sequencer::NavigationTool
+{
+
+#define UE_NAVIGATIONTOOL_DECLARE_VIEW_MODEL_TYPE_ID(Type)							\
+	static ::UE::Sequencer::TNavigationToolAutoRegisterViewModelTypeID<Type> ID;	\
+	static void RegisterTypeID();
+
+#define UE_NAVIGATIONTOOL_DECLARE_VIEW_MODEL_TYPE_ID_API(MODULE_API, Type)					\
+	MODULE_API static ::UE::Sequencer::TNavigationToolAutoRegisterViewModelTypeID<Type> ID;	\
+	MODULE_API static void RegisterTypeID();
+
+#define UE_NAVIGATIONTOOL_DEFINE_VIEW_MODEL_TYPE_ID(Type)																			\
+	::UE::Sequencer::TNavigationToolAutoRegisterViewModelTypeID<Type> Type::ID;														\
+	void Type::RegisterTypeID()																										\
+	{																																\
+		Type::ID.ID        = ::UE::Sequencer::FNavigationToolViewModelTypeID::RegisterNewID();										\
+		Type::ID.TypeTable = ::UE::Sequencer::FNavigationToolCastableTypeTable::MakeTypeTable<Type>((Type*)0, Type::ID.ID, #Type);	\
+	}
+
+struct FNavigationToolCastableTypeTable;
+
+struct FNavigationToolViewModelTypeID
+{
+	friend uint32 GetTypeHash(FNavigationToolViewModelTypeID In)
+	{
+		return In.ID;
+	}
+
+	friend bool operator<(FNavigationToolViewModelTypeID A, FNavigationToolViewModelTypeID B)
+	{
+		return A.ID < B.ID;
+	}
+
+	friend bool operator==(FNavigationToolViewModelTypeID A, FNavigationToolViewModelTypeID B)
+	{
+		return A.ID == B.ID;
+	}
+
+	uint32 GetTypeID() const
+	{
+		return ID;
+	}
+
+	SEQUENCERCORE_API static uint32 RegisterNewID();
+
+	FNavigationToolViewModelTypeID(FNavigationToolCastableTypeTable* InTypeTable, uint32 InID)
+		: TypeTable(InTypeTable)
+		, ID(InID)
+	{}
+
+private:
+	FNavigationToolCastableTypeTable* TypeTable = nullptr;
+	uint32 ID;
+};
+
+template<typename T>
+struct TNavigationToolViewModelTypeID
+{
+	operator FNavigationToolViewModelTypeID() const
+	{
+		Register();
+		return FNavigationToolViewModelTypeID{ TypeTable, ID };
+	}
+
+	uint32 GetTypeID() const
+	{
+		Register();
+		return ID;
+	}
+
+	FNavigationToolCastableTypeTable* GetTypeTable() const
+	{
+		Register();
+		return TypeTable;
+	}
+
+	void Register() const
+	{
+		if (!IsRegistered())
+		{
+			T::RegisterTypeID();
+		}
+	}
+
+protected:
+	friend T;
+
+	bool IsRegistered() const
+	{
+		return ID != ~0u;
+	}
+
+	FNavigationToolCastableTypeTable* TypeTable = nullptr;
+	uint32 ID = ~0u;
+};
+
+template<typename T>
+struct TNavigationToolAutoRegisterViewModelTypeID : TNavigationToolViewModelTypeID<T>
+{
+};
+
+} // namespace UE::Sequencer::NavigationTool
diff --git a/Engine/Plugins/Experimental/SequenceNavigator/Source/SequenceNavigator/Public/Providers/NavigationToolProvider.h b/Engine/Plugins/Experimental/SequenceNavigator/Source/SequenceNavigator/Public/Providers/NavigationToolProvider.h
new file mode 100644
--- /dev/null
+++ b/Engine/Plugins/Experimental/SequenceNavigator/Source/SequenceNavigator/Public/Providers/NavigationToolProvider.h
@@ -0,0 +1,70 @@
+﻿// Copyright Epic Games, Inc. All Rights Reserved.
+
+#pragma once
+
+#include "NavigationToolDefines.h"
+#include "Providers/INavigationToolProvider.h"
+
+class UMovieSceneSequence;
+struct FNavigationToolBuiltInFilterParams;
+struct FNavigationToolSerializedTree;
+struct FNavigationToolViewSaveState;
+
+namespace UE::Sequencer::NavigationTool
+{
+
+class FNavigationTool;
+class INavigationTool;
+
+/** Base Navigation Tool Provider to extend from */
+class FNavigationToolProvider : public INavigationToolProvider
+{
+public:
+	//~ Begin INavigationToolProvider
+
+	SEQUENCENAVIGATOR_API virtual void OnExtendColumnViews(TSet<FNavigationToolColumnView>& OutColumnViews) override;
+	SEQUENCENAVIGATOR_API virtual void OnExtendBuiltInFilters(TArray<FNavigationToolBuiltInFilterParams>& OutFilterParams) override;
+
+	SEQUENCENAVIGATOR_API virtual bool ShouldLockTool() const override { return false; }
+	SEQUENCENAVIGATOR_API virtual bool ShouldHideItem(const FNavigationToolItemPtr& InItem) const override { return false; }
+
+	SEQUENCENAVIGATOR_API virtual void UpdateItemIdContexts(const INavigationTool& InTool) final override;
+
+	SEQUENCENAVIGATOR_API virtual TOptional<EItemDropZone> OnToolItemCanAcceptDrop(const FDragDropEvent& InDragDropEvent, EItemDropZone InDropZone, const FNavigationToolItemPtr& InTargetItem) const override;
+	SEQUENCENAVIGATOR_API virtual FReply OnToolItemAcceptDrop(const FDragDropEvent& InDragDropEvent, EItemDropZone InDropZone, const FNavigationToolItemPtr& InTargetItem) override;
+
+	//~ End INavigationToolProvider
+
+	/** @return Editor only saved data for a specific view */
+	FNavigationToolViewSaveState* GetViewSaveState(const INavigationTool& InTool, const int32 InToolViewId) const;
+
+	/** Ensure saved editor only data contains an entry for a specific view, creating it if necessary */
+	void EnsureToolViewCount(const INavigationTool& InTool, const int32 InToolViewId);
+
+	bool IsSequenceSupported(UMovieSceneSequence* const InSequence) const;
+
+private:
+	friend class FNavigationToolExtender;
+	friend class FNavigationTool;
+
+	void Activate(FNavigationTool& InTool);
+	void Deactivate(FNavigationTool& InTool);
+
+	void SaveState(FNavigationTool& InTool);
+	void LoadState(FNavigationTool& InTool);
+
+	void SaveSerializedTree(FNavigationTool& InTool, const bool bInResetTree);
+
+	void SaveSerializedTreeRecursive(const FNavigationToolItemPtr& InParentItem
+		, FNavigationToolSerializedTree& InSerializedTree);
+
+	void LoadSerializedTree(const FNavigationToolItemPtr& InParentItem
+		, FNavigationToolSerializedTree* const InSerializedTree);
+
+	void CleanupExtendedColumnViews();
+
+	TArray<FText> ExtendedColumnViewNames;
+	TArray<FName> ExtendedBuiltInFilterNames;
+};
+
+}
diff --git a/Engine/Plugins/Experimental/SequenceNavigator/Source/SequenceNavigator/Public/Widgets/SNavigationToolTreeRow.h b/Engine/Plugins/Experimental/SequenceNavigator/Source/SequenceNavigator/Public/Widgets/SNavigationToolTreeRow.h
new file mode 100644
--- /dev/null
+++ b/Engine/Plugins/Experimental/SequenceNavigator/Source/SequenceNavigator/Public/Widgets/SNavigationToolTreeRow.h
@@ -0,0 +1,52 @@
+// Copyright Epic Games, Inc. All Rights Reserved.
+
+#pragma once
+
+#include "NavigationToolDefines.h"
+#include "Widgets/Views/STreeView.h"
+#include "Widgets/Views/STableRow.h"
+
+namespace UE::Sequencer::NavigationTool
+{
+
+class FNavigationToolItem;
+class FNavigationToolView;
+class SNavigationToolTreeView;
+
+class SNavigationToolTreeRow : public SMultiColumnTableRow<FNavigationToolItemPtr>
+{
+public:
+	SLATE_BEGIN_ARGS(SNavigationToolTreeRow) {}
+		SLATE_ATTRIBUTE(FText, HighlightText)
+	SLATE_END_ARGS()
+
+	void Construct(const FArguments& InArgs
+		, const TSharedRef<FNavigationToolView>& InToolView
+		, const TSharedPtr<SNavigationToolTreeView>& InTreeView
+		, const FNavigationToolItemPtr& InItem);
+
+	virtual TSharedRef<SWidget> GenerateWidgetForColumn(const FName& InColumnName) override;
+
+	TAttribute<FText> GetHighlightText() const { return HighlightText; }
+
+	virtual FReply OnMouseButtonUp(const FGeometry& MyGeometry, const FPointerEvent& MouseEvent) override;
+	virtual FReply OnMouseButtonDoubleClick(const FGeometry& InGeometry, const FPointerEvent& InMouseEvent) override;
+
+	TSharedPtr<FNavigationToolView> GetToolView() const;
+
+	const FTableRowStyle* GetStyle() const { return Style; }
+
+	/** The default reply if a row did not handle AcceptDrop */
+	FReply OnDefaultDrop(const FDragDropEvent& InDragDropEvent) const;
+
+private:
+	FNavigationToolItemPtr Item;
+	
+	TWeakPtr<SNavigationToolTreeView> WeakTreeView;
+	
+	TWeakPtr<FNavigationToolView> WeakToolView;
+	
+	TAttribute<FText> HighlightText;
+};
+
+}
diff --git a/Engine/Plugins/MovieScene/LevelSequenceEditor/LevelSequenceEditor.uplugin b/Engine/Plugins/MovieScene/LevelSequenceEditor/LevelSequenceEditor.uplugin
--- a/Engine/Plugins/MovieScene/LevelSequenceEditor/LevelSequenceEditor.uplugin
+++ b/Engine/Plugins/MovieScene/LevelSequenceEditor/LevelSequenceEditor.uplugin
@@ -28,6 +28,10 @@
 		{
 			"Name": "SequencerScripting",
 			"Enabled": true
+		},
+		{
+			"Name": "SequenceNavigator",
+			"Enabled": true
 		}
 	]
 }
\ No newline at end of file
diff --git a/Engine/Plugins/MovieScene/LevelSequenceEditor/Source/LevelSequenceEditor/LevelSequenceEditor.Build.cs b/Engine/Plugins/MovieScene/LevelSequenceEditor/Source/LevelSequenceEditor/LevelSequenceEditor.Build.cs
--- a/Engine/Plugins/MovieScene/LevelSequenceEditor/Source/LevelSequenceEditor/LevelSequenceEditor.Build.cs
+++ b/Engine/Plugins/MovieScene/LevelSequenceEditor/Source/LevelSequenceEditor/LevelSequenceEditor.Build.cs
@@ -50,6 +50,7 @@ public class LevelSequenceEditor : ModuleRules
 				"MovieSceneTracks",
 				"Projects",
                 "PropertyEditor",
+                "SequenceNavigator",
 				"SequencerCore",
 				"Sequencer",
                 "Slate",
diff --git a/Engine/Plugins/MovieScene/LevelSequenceEditor/Source/LevelSequenceEditor/Private/LevelSequenceEditorToolkit.cpp b/Engine/Plugins/MovieScene/LevelSequenceEditor/Source/LevelSequenceEditor/Private/LevelSequenceEditorToolkit.cpp
--- a/Engine/Plugins/MovieScene/LevelSequenceEditor/Source/LevelSequenceEditor/Private/LevelSequenceEditorToolkit.cpp
+++ b/Engine/Plugins/MovieScene/LevelSequenceEditor/Source/LevelSequenceEditor/Private/LevelSequenceEditorToolkit.cpp
@@ -45,7 +45,7 @@
 #include "SequencerUtilities.h"
 #include "ViewportToolbar/UnrealEdViewportToolbarContext.h"
 #include "LevelSequenceNavigationToolProvider.h"
-#include "Tools/NavigationTool/Customization/NavigationToolExtender.h"
+#include "NavigationToolExtender.h"
 
 // @todo sequencer: hack: setting defaults for transform tracks
 
@@ -250,7 +250,6 @@ void FLevelSequenceEditorToolkit::Initialize(const EToolkitMode::Type Mode, cons
 		SequencerInitParams.HostCapabilities.bSupportsAddFromContentBrowser = true;
 		SequencerInitParams.HostCapabilities.bSupportsSidebar = true;
 		SequencerInitParams.HostCapabilities.bSupportsViewportSelectability = true;
-		SequencerInitParams.HostCapabilities.bSupportsNavigationTool = true;
 	}
 
 	ExtendSequencerToolbar("Sequencer.MainToolBar");
@@ -319,16 +318,18 @@ void FLevelSequenceEditorToolkit::Initialize(const EToolkitMode::Type Mode, cons
 		}
 	}
 
-	using namespace UE::Sequencer::NavigationTool;
-	if (const TSharedPtr<INavigationTool> Tool = Sequencer->GetNavigationTool())
+	if (Sequencer.IsValid())
 	{
-		const FName ToolId = FNavigationToolExtender::GetToolInstanceId(*Tool);
+		using namespace UE::Sequencer::NavigationTool;
 
-		const TSharedPtr<FNavigationToolProvider> FoundProvider = FNavigationToolExtender::FindToolProvider(ToolId, FLevelSequenceNavigationToolProvider::Identifier);
+		const FName ToolId = FNavigationToolExtender::GetToolInstanceId(*Sequencer);
+
+		const TSharedPtr<FNavigationToolProvider> FoundProvider = FNavigationToolExtender::FindToolProvider(ToolId
+			, FLevelSequenceNavigationToolProvider::Identifier);
 		if (!FoundProvider.IsValid())
 		{
 			NavigationToolProvider = MakeShared<FLevelSequenceNavigationToolProvider>();
-			FNavigationToolExtender::RegisterToolProvider(ToolId, NavigationToolProvider.ToSharedRef());
+			FNavigationToolExtender::RegisterToolProvider(Sequencer.ToSharedRef(), NavigationToolProvider.ToSharedRef());
 		}
 	}
 }
@@ -996,14 +997,13 @@ void FLevelSequenceEditorToolkit::SaveAsset_Execute()
 
 void FLevelSequenceEditorToolkit::OnClose()
 {
-	if (NavigationToolProvider.IsValid())
+	if (Sequencer.IsValid() && NavigationToolProvider.IsValid())
 	{
 		using namespace UE::Sequencer::NavigationTool;
-		if (const TSharedPtr<INavigationTool> Tool = Sequencer->GetNavigationTool())
-		{
-			const FName ToolId = FNavigationToolExtender::GetToolInstanceId(*Tool);
-			FNavigationToolExtender::UnregisterToolProvider(ToolId, NavigationToolProvider->GetIdentifier());
-		}
+
+		const FName ToolId = FNavigationToolExtender::GetToolInstanceId(*Sequencer);
+		FNavigationToolExtender::UnregisterToolProvider(ToolId, NavigationToolProvider->GetIdentifier());
+
 		NavigationToolProvider.Reset();
 	}
 
diff --git a/Engine/Plugins/MovieScene/LevelSequenceEditor/Source/LevelSequenceEditor/Private/LevelSequenceEditorToolkit.h b/Engine/Plugins/MovieScene/LevelSequenceEditor/Source/LevelSequenceEditor/Private/LevelSequenceEditorToolkit.h
--- a/Engine/Plugins/MovieScene/LevelSequenceEditor/Source/LevelSequenceEditor/Private/LevelSequenceEditorToolkit.h
+++ b/Engine/Plugins/MovieScene/LevelSequenceEditor/Source/LevelSequenceEditor/Private/LevelSequenceEditorToolkit.h
@@ -9,6 +9,7 @@ class FSpawnTabArgs;
 struct FFrameNumber;
 
 class AActor;
+class FLevelSequenceNavigationToolProvider;
 class FMenuBuilder;
 class FToolBarBuilder;
 class IAssetViewport;
@@ -21,11 +22,6 @@ class FLevelSequencePlaybackContext;
 class UPrimitiveComponent;
 enum class EMapChangeType : uint8;
 
-namespace UE::Sequencer::NavigationTool
-{
-	class FLevelSequenceNavigationToolProvider;
-}
-
 /**
  * Implements an Editor toolkit for level sequences.
  */
@@ -152,7 +148,8 @@ private:
 	/** The sequencer used by this editor. */
 	TSharedPtr<ISequencer> Sequencer;
 
-	TSharedPtr<UE::Sequencer::NavigationTool::FLevelSequenceNavigationToolProvider> NavigationToolProvider;
+	/** The provider that supplies data and extends the Navigation Tool */
+	TSharedPtr<FLevelSequenceNavigationToolProvider> NavigationToolProvider;
 
 	/** Pointer to the style set to use for toolkits. */
 	TSharedRef<ISlateStyle> Style;
diff --git a/Engine/Plugins/MovieScene/LevelSequenceEditor/Source/LevelSequenceEditor/Private/LevelSequenceNavigationToolProvider.cpp b/Engine/Plugins/MovieScene/LevelSequenceEditor/Source/LevelSequenceEditor/Private/LevelSequenceNavigationToolProvider.cpp
--- a/Engine/Plugins/MovieScene/LevelSequenceEditor/Source/LevelSequenceEditor/Private/LevelSequenceNavigationToolProvider.cpp
+++ b/Engine/Plugins/MovieScene/LevelSequenceEditor/Source/LevelSequenceEditor/Private/LevelSequenceNavigationToolProvider.cpp
@@ -1,37 +1,34 @@
 ﻿// Copyright Epic Games, Inc. All Rights Reserved.
 
 #include "LevelSequenceNavigationToolProvider.h"
+#include "Columns/NavigationToolColorColumn.h"
+#include "Columns/NavigationToolColumnExtender.h"
+#include "Columns/NavigationToolCommentColumn.h"
+#include "Columns/NavigationToolHBiasColumn.h"
+#include "Columns/NavigationToolInTimeColumn.h"
+#include "Columns/NavigationToolItemsColumn.h"
+#include "Columns/NavigationToolLabelColumn.h"
+#include "Columns/NavigationToolLengthColumn.h"
+#include "Columns/NavigationToolLockColumn.h"
+#include "Columns/NavigationToolMarkerVisibilityColumn.h"
+#include "Columns/NavigationToolOutTimeColumn.h"
+#include "Columns/NavigationToolPlayheadColumn.h"
+#include "Columns/NavigationToolRevisionControlColumn.h"
+#include "Columns/NavigationToolStartFrameOffsetColumn.h"
+#include "Columns/NavigationToolTakeColumn.h"
+#include "Columns/NavigationToolDeactiveStateColumn.h"
+#include "Filters/Filters/NavigationToolBuiltInFilterParams.h"
+#include "Filters/Filters/NavigationToolFilterBase.h"
+#include "INavigationTool.h"
 #include "ISequencer.h"
 #include "ISourceControlModule.h"
+#include "Items/NavigationToolSequence.h"
 #include "LevelSequence.h"
+#include "NavigationToolSettings.h"
 #include "MovieScene.h"
-#include "Tools/NavigationTool/Columns/NavigationToolColorColumn.h"
-#include "Tools/NavigationTool/Columns/NavigationToolColumnExtender.h"
-#include "Tools/NavigationTool/Columns/NavigationToolCommentColumn.h"
-#include "Tools/NavigationTool/Columns/NavigationToolHBiasColumn.h"
-#include "Tools/NavigationTool/Columns/NavigationToolInTimeColumn.h"
-#include "Tools/NavigationTool/Columns/NavigationToolItemsColumn.h"
-#include "Tools/NavigationTool/Columns/NavigationToolLabelColumn.h"
-#include "Tools/NavigationTool/Columns/NavigationToolLengthColumn.h"
-#include "Tools/NavigationTool/Columns/NavigationToolLockColumn.h"
-#include "Tools/NavigationTool/Columns/NavigationToolMarkerVisibilityColumn.h"
-#include "Tools/NavigationTool/Columns/NavigationToolOutTimeColumn.h"
-#include "Tools/NavigationTool/Columns/NavigationToolPlayheadColumn.h"
-#include "Tools/NavigationTool/Columns/NavigationToolRevisionControlColumn.h"
-#include "Tools/NavigationTool/Columns/NavigationToolStartFrameOffsetColumn.h"
-#include "Tools/NavigationTool/Columns/NavigationToolTakeColumn.h"
-#include "Tools/NavigationTool/Columns/NavigationToolDeactiveStateColumn.h"
-#include "Tools/NavigationTool/Filters/Filters/NavigationToolBuiltInFilterParams.h"
-#include "Tools/NavigationTool/Filters/Filters/NavigationToolFilterBase.h"
-#include "Tools/NavigationTool/INavigationTool.h"
-#include "Tools/NavigationTool/Items/NavigationToolSequence.h"
-#include "Tools/NavigationTool/NavigationToolSettings.h"
 
 #define LOCTEXT_NAMESPACE "LevelSequenceNavigationToolProvider"
 
-namespace UE::Sequencer::NavigationTool
-{
-
 const FName FLevelSequenceNavigationToolProvider::Identifier = TEXT("LevelSequence");
 
 const FText FLevelSequenceNavigationToolProvider::AnimationColumnViewName = LOCTEXT("AnimationColumnViewName", "Animation");
@@ -81,8 +78,10 @@ void FLevelSequenceNavigationToolProvider::OnDeactivate()
 		.Remove(DragDropInitializedDelegate);*/
 }
 
-void FLevelSequenceNavigationToolProvider::OnExtendColumns(FNavigationToolColumnExtender& OutExtender)
+void FLevelSequenceNavigationToolProvider::OnExtendColumns(UE::Sequencer::NavigationTool::FNavigationToolColumnExtender& OutExtender)
 {
+	using namespace UE::Sequencer::NavigationTool;
+
 	OutExtender.AddColumn<FNavigationToolPlayheadColumn>();
 	OutExtender.AddColumn<FNavigationToolDeactiveStateColumn>();
 	OutExtender.AddColumn<FNavigationToolMarkerVisibilityColumn>();
@@ -109,6 +108,8 @@ void FLevelSequenceNavigationToolProvider::OnExtendColumns(FNavigationToolColumn
 
 void FLevelSequenceNavigationToolProvider::OnExtendColumnViews(TSet<FNavigationToolColumnView>& OutColumnViews)
 {
+	using namespace UE::Sequencer::NavigationTool;
+
 	FNavigationToolColumnView& AnimationColumnView = OutColumnViews
 		.FindOrAdd(FNavigationToolColumnView(AnimationColumnViewName));
 
@@ -126,11 +127,13 @@ void FLevelSequenceNavigationToolProvider::OnExtendColumnViews(TSet<FNavigationT
 	FNavigationToolProvider::OnExtendColumnViews(OutColumnViews);
 }
 
-void FLevelSequenceNavigationToolProvider::OnExtendItemChildren(INavigationTool& InTool
-	, const FNavigationToolItemPtr& InParentItem
-	, TArray<FNavigationToolItemPtr>& OutChildren
+void FLevelSequenceNavigationToolProvider::OnExtendItemChildren(UE::Sequencer::NavigationTool::INavigationTool& InTool
+	, const UE::Sequencer::NavigationTool::FNavigationToolItemPtr& InParentItem
+	, TArray<UE::Sequencer::NavigationTool::FNavigationToolItemPtr>& OutChildren
 	, const bool bInRecursive)
 {
+	using namespace UE::Sequencer::NavigationTool;
+
 	FNavigationToolProvider::OnExtendItemChildren(InTool, InParentItem, OutChildren, bInRecursive);
 
 	// Only extending root item
@@ -172,7 +175,7 @@ void FLevelSequenceNavigationToolProvider::OnExtendBuiltInFilters(TArray<FNaviga
 	FNavigationToolProvider::OnExtendBuiltInFilters(OutFilterParams);
 }
 
-FMovieSceneEditorData* FLevelSequenceNavigationToolProvider::GetRootMovieSceneEditorData(const INavigationTool& InTool) const
+FMovieSceneEditorData* FLevelSequenceNavigationToolProvider::GetRootMovieSceneEditorData(const UE::Sequencer::NavigationTool::INavigationTool& InTool) const
 {
 	const TSharedPtr<ISequencer> Sequencer = InTool.GetSequencer();
 	if (!Sequencer.IsValid())
@@ -195,7 +198,7 @@ FMovieSceneEditorData* FLevelSequenceNavigationToolProvider::GetRootMovieSceneEd
 	return &RootMovieScene->GetEditorData();
 }
 
-FNavigationToolSaveState* FLevelSequenceNavigationToolProvider::GetSaveState(const INavigationTool& InTool) const
+FNavigationToolSaveState* FLevelSequenceNavigationToolProvider::GetSaveState(const UE::Sequencer::NavigationTool::INavigationTool& InTool) const
 {
 	if (FMovieSceneEditorData* const EditorData = GetRootMovieSceneEditorData(InTool))
 	{
@@ -204,7 +207,8 @@ FNavigationToolSaveState* FLevelSequenceNavigationToolProvider::GetSaveState(con
 	return nullptr;
 }
 
-void FLevelSequenceNavigationToolProvider::SetSaveState(const INavigationTool& InTool, const FNavigationToolSaveState& InSaveState) const
+void FLevelSequenceNavigationToolProvider::SetSaveState(const UE::Sequencer::NavigationTool::INavigationTool& InTool
+	, const FNavigationToolSaveState& InSaveState) const
 {
 	if (FMovieSceneEditorData* const EditorData = GetRootMovieSceneEditorData(InTool))
 	{
@@ -212,6 +216,4 @@ void FLevelSequenceNavigationToolProvider::SetSaveState(const INavigationTool& I
 	}
 }
 
-}
-
 #undef LOCTEXT_NAMESPACE
diff --git a/Engine/Plugins/MovieScene/LevelSequenceEditor/Source/LevelSequenceEditor/Private/LevelSequenceNavigationToolProvider.h b/Engine/Plugins/MovieScene/LevelSequenceEditor/Source/LevelSequenceEditor/Private/LevelSequenceNavigationToolProvider.h
--- a/Engine/Plugins/MovieScene/LevelSequenceEditor/Source/LevelSequenceEditor/Private/LevelSequenceNavigationToolProvider.h
+++ b/Engine/Plugins/MovieScene/LevelSequenceEditor/Source/LevelSequenceEditor/Private/LevelSequenceNavigationToolProvider.h
@@ -2,21 +2,25 @@
 
 #pragma once
 
+#include "NavigationToolDefines.h"
+#include "Providers/NavigationToolProvider.h"
 #include "Templates/SharedPointer.h"
-#include "Tools/NavigationTool/NavigationToolDefines.h"
-#include "Tools/NavigationTool/Providers/NavigationToolProvider.h"
 #include "UObject/WeakObjectPtr.h"
 
+struct FNavigationToolColumnView;
 class FUICommandList;
 struct FMovieSceneEditorData;
+struct FNavigationToolBuiltInFilterParams;
+struct FNavigationToolSaveState;
 
 namespace UE::Sequencer::NavigationTool
 {
+	class FNavigationToolColumnExtender;
+	class FNavigationToolItem;
+	class INavigationTool;
+}
 
-class FNavigationToolItem;
-class INavigationTool;
-
-class FLevelSequenceNavigationToolProvider : public FNavigationToolProvider
+class FLevelSequenceNavigationToolProvider : public UE::Sequencer::NavigationTool::FNavigationToolProvider
 {
 public:
 	static const FName Identifier;
@@ -39,23 +43,22 @@ public:
 	virtual void OnActivate() override;
 	virtual void OnDeactivate() override;
 
-	virtual void OnExtendColumns(FNavigationToolColumnExtender& OutExtender) override;
+	virtual void OnExtendColumns(UE::Sequencer::NavigationTool::FNavigationToolColumnExtender& OutExtender) override;
 	virtual void OnExtendColumnViews(TSet<FNavigationToolColumnView>& OutColumnViews) override;
-	virtual void OnExtendItemChildren(INavigationTool& InTool
-		, const FNavigationToolItemPtr& InParentItem
-		, TArray<FNavigationToolItemPtr>& OutChildren
+	virtual void OnExtendItemChildren(UE::Sequencer::NavigationTool::INavigationTool& InTool
+		, const UE::Sequencer::NavigationTool::FNavigationToolItemPtr& InParentItem
+		, TArray<UE::Sequencer::NavigationTool::FNavigationToolItemPtr>& OutChildren
 		, const bool bInRecursive) override;
 	virtual void OnExtendBuiltInFilters(TArray<FNavigationToolBuiltInFilterParams>& OutFilterParams) override;
 
-	virtual FNavigationToolSaveState* GetSaveState(const INavigationTool& InTool) const override;
-	virtual void SetSaveState(const INavigationTool& InTool, const FNavigationToolSaveState& InSaveState) const override;
+	virtual FNavigationToolSaveState* GetSaveState(const UE::Sequencer::NavigationTool::INavigationTool& InTool) const override;
+	virtual void SetSaveState(const UE::Sequencer::NavigationTool::INavigationTool& InTool
+		, const FNavigationToolSaveState& InSaveState) const override;
 
 	//~ End INavigationToolProvider
 
 private:
-	FMovieSceneEditorData* GetRootMovieSceneEditorData(const INavigationTool& InTool) const;
+	FMovieSceneEditorData* GetRootMovieSceneEditorData(const UE::Sequencer::NavigationTool::INavigationTool& InTool) const;
 
 	TSharedRef<FUICommandList> ToolCommands;
 };
-
-}
diff --git a/Engine/Source/Editor/Sequencer/Private/Filters/SequencerFilterBar.cpp b/Engine/Source/Editor/Sequencer/Private/Filters/SequencerFilterBar.cpp
--- a/Engine/Source/Editor/Sequencer/Private/Filters/SequencerFilterBar.cpp
+++ b/Engine/Source/Editor/Sequencer/Private/Filters/SequencerFilterBar.cpp
@@ -19,6 +19,7 @@
 #include "Filters/SequencerTrackFilterCollection.h"
 #include "Filters/SequencerTrackFilterCommands.h"
 #include "Filters/SequencerTrackFilterExtension.h"
+#include "Filters/Widgets/SFilterBarIsolateHideShow.h"
 #include "Filters/Widgets/SFilterExpressionHelpDialog.h"
 #include "Filters/Widgets/SSequencerFilterBar.h"
 #include "Filters/Widgets/SSequencerSearchBox.h"
@@ -34,7 +35,6 @@
 #include "SSequencer.h"
 #include "Stats/StatsMisc.h"
 #include "UObject/UObjectIterator.h"
-#include "Widgets/SFilterBarIsolateHideShow.h"
 
 using namespace UE::Sequencer;
 
diff --git a/Engine/Source/Editor/Sequencer/Private/SSequencer.cpp b/Engine/Source/Editor/Sequencer/Private/SSequencer.cpp
--- a/Engine/Source/Editor/Sequencer/Private/SSequencer.cpp
+++ b/Engine/Source/Editor/Sequencer/Private/SSequencer.cpp
@@ -143,7 +143,6 @@
 #include "ContentBrowserModule.h"
 #include "Menus/SequencerToolbarUtils.h"
 #include "Settings/EditorExperimentalSettings.h"
-#include "Tools/NavigationTool/NavigationToolStyle.h"
 
 #define LOCTEXT_NAMESPACE "Sequencer"
 
@@ -1672,18 +1671,6 @@ void SSequencer::PopulateToolBar(UToolMenu* InMenu)
 			CurveEditorSection.AddEntry(ShowCurveEditorEntry);
 		}
 	}
-
-	if (Sequencer->GetHostCapabilities().bSupportsNavigationTool)
-	{
-		FToolMenuSection& NavigationToolSection = InMenu->FindOrAddSection(TEXT("NavigationTool"));
-
-		FToolMenuEntry ShowNavigationToolEntry = FToolMenuEntry::InitToolBarButton(
-			FSequencerCommands::Get().ToggleNavigationToolVisible,
-			TAttribute<FText>(),TAttribute<FText>(),
-			FSlateIcon(FNavigationToolStyle::Get().GetStyleSetName(), TEXT("Icon.ToolBar")));
-		ShowNavigationToolEntry.StyleNameOverride = SequencerToolbarStyleName;
-		NavigationToolSection.AddEntry(ShowNavigationToolEntry);
-	}
 }
 
 TSharedRef<SWidget> SSequencer::MakeAddMenu()
diff --git a/Engine/Source/Editor/Sequencer/Private/Sequencer.cpp b/Engine/Source/Editor/Sequencer/Private/Sequencer.cpp
--- a/Engine/Source/Editor/Sequencer/Private/Sequencer.cpp
+++ b/Engine/Source/Editor/Sequencer/Private/Sequencer.cpp
@@ -172,8 +172,6 @@
 #include "Editor/TransBuffer.h"
 #include "Filters/SequencerFilterBar.h"
 #include "Sidebar/SidebarDrawerConfig.h"
-#include "Tools/NavigationTool/INavigationTool.h"
-#include "Tools/NavigationTool/NavigationTool.h"
 #include "Widgets/Sidebar/SequencerSelectionDrawer.h"
 #include "Misc/Thumbnail/ThumbnailCaptureUtils.h"
 #include "ToolMenus.h"
@@ -728,14 +726,6 @@ void FSequencer::InitSequencer(const FSequencerInitParams& InitParams
 	// Update the view range to the new current time
 	UpdateTimeBoundsToFocusedMovieScene();
 
-	// NavigationTool must be created before BindCommands()
-	if (HostCapabilities.bSupportsNavigationTool)
-	{
-		NavigationTool = MakeShared<NavigationTool::FNavigationTool>(SharedThisRef);
-		NavigationTool->SetBaseCommandList(SequencerCommandBindings);
-		NavigationTool->Init();
-	}
-
 	// NOTE: Could fill in asset editor commands here!
 
 	BindCommands();
@@ -12020,17 +12010,6 @@ void FSequencer::BindCommands()
 		Commands.ToggleTrackSelectionSolo,
 		FExecuteAction::CreateSP(this, &FSequencer::ToggleTrackSelectionSolo));
 
-	if (NavigationTool.IsValid())
-	{
-		const TSharedRef<UE::Sequencer::NavigationTool::FNavigationTool> NavigationToolRef = NavigationTool.ToSharedRef();
-
-		SequencerCommandBindings->MapAction(
-			Commands.ToggleNavigationToolVisible,
-			FExecuteAction::CreateSP(NavigationToolRef, &UE::Sequencer::NavigationTool::FNavigationTool::ToggleToolTabVisible),
-			FCanExecuteAction(),
-			FIsActionChecked::CreateSP(NavigationToolRef, &UE::Sequencer::NavigationTool::FNavigationTool::IsToolTabVisible));
-	}
-
 	// If this sequencer supports a curve editor, let's add bindings for it.
 	FCurveEditorExtension* CurveEditorExtension = ViewModel->CastDynamic<FCurveEditorExtension>();
 	if (CurveEditorExtension && ensure(CurveEditorExtension->GetCurveEditor()))
@@ -12983,11 +12962,6 @@ void FSequencer::ToggleTrackSelectionSolo()
 	}
 }
 
-TSharedPtr<UE::Sequencer::NavigationTool::INavigationTool> FSequencer::GetNavigationTool() const
-{
-	return NavigationTool;
-}
-
 TSharedPtr<UE::Sequencer::SOutlinerView> FSequencer::GetOutlinerViewWidget() const
 {
 	return SequencerWidget->GetTreeView();
diff --git a/Engine/Source/Editor/Sequencer/Private/Sequencer.h b/Engine/Source/Editor/Sequencer/Private/Sequencer.h
--- a/Engine/Source/Editor/Sequencer/Private/Sequencer.h
+++ b/Engine/Source/Editor/Sequencer/Private/Sequencer.h
@@ -43,7 +43,6 @@
 #include "SequencerTimeChangeUndoRedoProxy.h"
 #include "Capabilities/CameraCutViewTargetCacheCapability.h"
 #include "Filters/ISequencerTrackFilters.h"
-#include "Tools/NavigationTool/INavigationTool.h"
 
 class AActor;
 class ACameraActor;
@@ -106,14 +105,6 @@ namespace UE
 		class FViewModel;
 		class FSequencerSelection;
 
-		namespace NavigationTool
-		{
-
-			class FNavigationTool;
-			class INavigationTool;
-
-		}
-
 	} // namespace Sequencer
 } // namespace UE
 
@@ -1279,8 +1270,6 @@ public:
 	/** Checks to see if a clock extension wants the standard PlayRate combo menu hidden */
 	EVisibility GetPlayRateComboVisibility() const;
 
-	virtual TSharedPtr<UE::Sequencer::NavigationTool::INavigationTool> GetNavigationTool() const override;
-
 	virtual TSharedPtr<UE::Sequencer::SOutlinerView> GetOutlinerViewWidget() const override;
 
 private:
@@ -1645,9 +1634,6 @@ private:
 
 	TSharedPtr<FSequencerFilterBar> FilterBar;
 
-	/** Navigation Tool instance for this Sequencer */
-	TSharedPtr<UE::Sequencer::NavigationTool::FNavigationTool> NavigationTool;
-
 	enum class ETimeControllerState
 	{
 		ReadyToPlay,
diff --git a/Engine/Source/Editor/Sequencer/Private/SequencerCommands.cpp b/Engine/Source/Editor/Sequencer/Private/SequencerCommands.cpp
--- a/Engine/Source/Editor/Sequencer/Private/SequencerCommands.cpp
+++ b/Engine/Source/Editor/Sequencer/Private/SequencerCommands.cpp
@@ -206,8 +206,6 @@ void FSequencerCommands::RegisterCommands()
 	UI_COMMAND(ToggleTrackSelectionDeactive, "Toggle Track Deactive", "Toggle the deactive state of Sequencer track selection", EUserInterfaceActionType::Button, FInputChord());
 	UI_COMMAND(ToggleTrackSelectionMute, "Toggle Track Mute", "Toggle the muted state of Sequencer track selection", EUserInterfaceActionType::Button, FInputChord());
 	UI_COMMAND(ToggleTrackSelectionSolo, "Toggle Track Solo", "Toggle the soloed state of Sequencer track selection", EUserInterfaceActionType::Button, FInputChord());
-
-	UI_COMMAND(ToggleNavigationToolVisible, "Sequence Navigation", "Show the Sequence Navigation Tool", EUserInterfaceActionType::ToggleButton, FInputChord());
 }
 
 #undef LOCTEXT_NAMESPACE
diff --git a/Engine/Source/Editor/Sequencer/Private/SequencerModule.cpp b/Engine/Source/Editor/Sequencer/Private/SequencerModule.cpp
--- a/Engine/Source/Editor/Sequencer/Private/SequencerModule.cpp
+++ b/Engine/Source/Editor/Sequencer/Private/SequencerModule.cpp
@@ -22,7 +22,6 @@
 #include "MovieSceneSignedObject.h"
 #include "MovieSceneSequence.h"
 #include "Sections/MovieSceneSubSection.h"
-#include "Tools/NavigationTool/Customization/NavigationToolExtender.h"
 
 #include "MVVM/CurveEditorExtension.h"
 #include "MVVM/CurveEditorIntegrationExtension.h"
@@ -710,11 +709,6 @@ public:
 		return ObjectSchemas;
 	}
 
-	UE::Sequencer::NavigationTool::FNavigationToolExtender& GetNavigationToolExtender() override
-	{
-		return NavigationToolExtender;
-	}
-
 private:
 
 	TSet<FAnimatedPropertyKey> PropertyAnimators;
@@ -789,8 +783,6 @@ private:
 	// Outliner Indicator Item Delegate Handles
 	FDelegateHandle ConditionOutlinerIndicatorHandle;
 	FDelegateHandle TimeWarpOutlinerIndicatorHandle;
-
-	UE::Sequencer::NavigationTool::FNavigationToolExtender NavigationToolExtender;
 };
 
 IMPLEMENT_MODULE(FSequencerModule, Sequencer);
diff --git a/Engine/Source/Editor/Sequencer/Public/Filters/SequencerFilterBarConfig.h b/Engine/Source/Editor/Sequencer/Public/Filters/SequencerFilterBarConfig.h
--- a/Engine/Source/Editor/Sequencer/Public/Filters/SequencerFilterBarConfig.h
+++ b/Engine/Source/Editor/Sequencer/Public/Filters/SequencerFilterBarConfig.h
@@ -30,26 +30,26 @@ struct FSequencerFilterBarConfig
 public:
 	/** Common Filters */
 
-	bool IsFilterEnabled(const FString& InFilterName) const;
-	bool SetFilterEnabled(const FString& InFilterName, const bool bInActive);
+	SEQUENCER_API bool IsFilterEnabled(const FString& InFilterName) const;
+	SEQUENCER_API bool SetFilterEnabled(const FString& InFilterName, const bool bInActive);
 
-	bool IsFilterActive(const FString& InFilterName) const;
-	bool SetFilterActive(const FString& InFilterName, const bool bInActive);
+	SEQUENCER_API bool IsFilterActive(const FString& InFilterName) const;
+	SEQUENCER_API bool SetFilterActive(const FString& InFilterName, const bool bInActive);
 
 	const FSequencerFilterSet& GetCommonActiveSet() const;
 
 	/** Custom Text Filters */
 
-	TArray<FCustomTextFilterData>& GetCustomTextFilters();
-	bool HasCustomTextFilter(const FString& InFilterName) const;
-	FCustomTextFilterData* FindCustomTextFilter(const FString& InFilterName);
-	bool AddCustomTextFilter(FCustomTextFilterData InFilterData);
-	bool RemoveCustomTextFilter(const FString& InFilterName);
+	SEQUENCER_API TArray<FCustomTextFilterData>& GetCustomTextFilters();
+	SEQUENCER_API bool HasCustomTextFilter(const FString& InFilterName) const;
+	SEQUENCER_API FCustomTextFilterData* FindCustomTextFilter(const FString& InFilterName);
+	SEQUENCER_API bool AddCustomTextFilter(FCustomTextFilterData InFilterData);
+	SEQUENCER_API bool RemoveCustomTextFilter(const FString& InFilterName);
 
 	/** Filter Bar Layout */
 
-	EFilterBarLayout GetFilterBarLayout() const;
-	void SetFilterBarLayout(const EFilterBarLayout InLayout);
+	SEQUENCER_API EFilterBarLayout GetFilterBarLayout() const;
+	SEQUENCER_API void SetFilterBarLayout(const EFilterBarLayout InLayout);
 
 protected:
 	/** The currently active set of common and custom text filters that should be restored on editor load */
diff --git a/Engine/Source/Editor/Sequencer/Public/Filters/SequencerFilterBase.h b/Engine/Source/Editor/Sequencer/Public/Filters/SequencerFilterBase.h
--- a/Engine/Source/Editor/Sequencer/Public/Filters/SequencerFilterBase.h
+++ b/Engine/Source/Editor/Sequencer/Public/Filters/SequencerFilterBase.h
@@ -7,8 +7,6 @@
 #include "Templates/SharedPointer.h"
 #include "MovieSceneSequence.h"
 
-struct FSlateBrush;
-
 /** Sequencer base filter for all Sequencer filters (Sequencer, Navigation Tool, etc.) */
 template<typename InFilterType>
 class FSequencerFilterBase
@@ -16,7 +14,7 @@ class FSequencerFilterBase
 	, public TSharedFromThis<FSequencerFilterBase<InFilterType>>
 {
 public:
-	SEQUENCER_API static bool SupportsLevelSequence(UMovieSceneSequence* const InSequence)
+	static bool SupportsLevelSequence(UMovieSceneSequence* const InSequence)
 	{
 		static UClass* const LevelSequenceClass = FindObject<UClass>(nullptr, TEXT("/Script/LevelSequence.LevelSequence"), true);
 		return IsValid(InSequence)
@@ -24,7 +22,7 @@ public:
 			&& InSequence->GetClass()->IsChildOf(LevelSequenceClass);
 	}
 
-	SEQUENCER_API static bool SupportsUMGSequence(UMovieSceneSequence* const InSequence)
+	static bool SupportsUMGSequence(UMovieSceneSequence* const InSequence)
 	{
 		static UClass* const WidgetAnimationClass = FindObject<UClass>(nullptr, TEXT("/Script/UMG.WidgetAnimation"), true);
 		return IsValid(InSequence)
@@ -32,7 +30,7 @@ public:
 			&& InSequence->GetClass()->IsChildOf(WidgetAnimationClass);
 	}
 
-	SEQUENCER_API static bool SupportsDaySequence(UMovieSceneSequence* const InSequence)
+	static bool SupportsDaySequence(UMovieSceneSequence* const InSequence)
 	{
 		static UClass* const DaySequenceClass = FindObject<UClass>(nullptr, TEXT("/Script/DaySequence.DaySequence"), true);
 		return IsValid(InSequence)
@@ -40,7 +38,7 @@ public:
 			&& InSequence->GetClass()->IsChildOf(DaySequenceClass);
 	}
 
-	SEQUENCER_API static FText BuildTooltipTextForCommand(const FText& InBaseText, const TSharedPtr<FUICommandInfo>& InCommand)
+	static FText BuildTooltipTextForCommand(const FText& InBaseText, const TSharedPtr<FUICommandInfo>& InCommand)
 	{
 		const TSharedRef<const FInputChord> FirstValidChord = InCommand->GetFirstValidChord();
 		if (FirstValidChord->IsValidChord())
@@ -57,14 +55,14 @@ public:
 	{}
 
 	//~ Begin IFilter
-	SEQUENCER_API virtual bool PassesFilter(InFilterType InItem) const override { return true; };
+	virtual bool PassesFilter(InFilterType InItem) const override { return true; };
 	//~ End IFilter
 
 	//~ Begin FFilterBase
 
-	SEQUENCER_API virtual FText GetDisplayName() const override { return FText::GetEmpty(); }
+	virtual FText GetDisplayName() const override { return FText::GetEmpty(); }
 
-	SEQUENCER_API virtual FText GetToolTipText() const override
+	virtual FText GetToolTipText() const override
 	{
 		if (const TSharedPtr<FUICommandInfo> ToggleCommand = GetToggleCommand())
 		{
@@ -73,7 +71,7 @@ public:
 		return GetDefaultToolTipText();
 	}
 
-	SEQUENCER_API virtual FLinearColor GetColor() const override { return FLinearColor::White; }
+	virtual FLinearColor GetColor() const override { return FLinearColor::White; }
 
 	virtual void ModifyContextMenu(FMenuBuilder& MenuBuilder) override {}
 
@@ -86,13 +84,13 @@ public:
 
 	//~ End FFilterBase
 
-	SEQUENCER_API ISequencer& GetSequencer() const { return FilterInterface.GetSequencer(); }
+	ISequencer& GetSequencer() const { return FilterInterface.GetSequencer(); }
 
-	SEQUENCER_API virtual ISequencerFilterBar& GetFilterInterface() const { return FilterInterface; }
+	virtual ISequencerFilterBar& GetFilterInterface() const { return FilterInterface; }
 
-	SEQUENCER_API virtual FText GetDefaultToolTipText() const { return FText(); }
-	SEQUENCER_API virtual FSlateIcon GetIcon() const { return FSlateIcon(); }
-	SEQUENCER_API virtual bool IsCustomTextFilter() const { return false; }
+	virtual FText GetDefaultToolTipText() const { return FText(); }
+	virtual FSlateIcon GetIcon() const { return FSlateIcon(); }
+	virtual bool IsCustomTextFilter() const { return false; }
 
 	virtual void BindCommands()
 	{
@@ -102,22 +100,22 @@ public:
 		}
 	}
 
-	SEQUENCER_API virtual TSharedPtr<FUICommandInfo> GetToggleCommand() const { return nullptr; }
+	virtual TSharedPtr<FUICommandInfo> GetToggleCommand() const { return nullptr; }
 
-	SEQUENCER_API bool CanToggleFilter() const
+	bool CanToggleFilter() const
 	{
 		const FString FilterName = GetDisplayName().ToString();
 		return FilterInterface.IsFilterActiveByDisplayName(FilterName);
 	}
 
-	SEQUENCER_API void ToggleFilter()
+	void ToggleFilter()
 	{
 		const FString FilterName = GetDisplayName().ToString();
 		const bool bNewState = !FilterInterface.IsFilterActiveByDisplayName(FilterName);
 		FilterInterface.SetFilterActiveByDisplayName(FilterName, bNewState);
 	}
 
-	SEQUENCER_API virtual bool SupportsSequence(UMovieSceneSequence* const InSequence) const { return true; }
+	virtual bool SupportsSequence(UMovieSceneSequence* const InSequence) const { return true; }
 
 protected:
 	void MapToggleAction(const TSharedPtr<FUICommandInfo>& InCommand)
diff --git a/Engine/Source/Editor/Sequencer/Public/ISequencer.h b/Engine/Source/Editor/Sequencer/Public/ISequencer.h
--- a/Engine/Source/Editor/Sequencer/Public/ISequencer.h
+++ b/Engine/Source/Editor/Sequencer/Public/ISequencer.h
@@ -96,13 +96,6 @@ struct FCreateBindingParams;
 struct FTimeDomainOverride;
 struct FSequencerNumericTypeInterface;
 
-namespace NavigationTool
-{
-
-class INavigationTool;
-
-}
-
 } // namespace Sequencer
 } // namespace UE
 
@@ -1001,8 +994,6 @@ public:
 
 	virtual bool TrackSupportsConditions(const UMovieSceneTrack* Track) const = 0;
 
-	virtual TSharedPtr<UE::Sequencer::NavigationTool::INavigationTool> GetNavigationTool() const = 0;
-
 	virtual TSharedPtr<UE::Sequencer::SOutlinerView> GetOutlinerViewWidget() const = 0;
 
 protected:
diff --git a/Engine/Source/Editor/Sequencer/Public/ISequencerModule.h b/Engine/Source/Editor/Sequencer/Public/ISequencerModule.h
--- a/Engine/Source/Editor/Sequencer/Public/ISequencerModule.h
+++ b/Engine/Source/Editor/Sequencer/Public/ISequencerModule.h
@@ -30,11 +30,6 @@ namespace UE::Sequencer
 	class IObjectSchema;
 	class IOutlinerColumn;
 	class IOutlinerIndicatorBuilder;
-
-	namespace NavigationTool
-	{
-		class FNavigationToolExtender;
-	}
 } // namespace UE::Sequencer
 
 enum class ECurveEditorTreeFilterType : uint32;
@@ -165,9 +160,6 @@ struct FSequencerHostCapabilities
 	/** Do we support selectability options in the viewport */
 	bool bSupportsViewportSelectability;
 
-	/** Do we support the Navigation Tool */
-	bool bSupportsNavigationTool;
-
 	FSequencerHostCapabilities()
 		: bSupportsSaveMovieSceneAsset(false)
 		, bSupportsCurveEditor(false)
@@ -176,7 +168,6 @@ struct FSequencerHostCapabilities
 		, bSupportsAddFromContentBrowser(false)
 		, bSupportsSidebar(false)
 		, bSupportsViewportSelectability(false)
-		, bSupportsNavigationTool(false)
 	{}
 };
 
@@ -535,9 +526,6 @@ public:
 	 */
 	virtual TArray<FString> GetMovieRendererNames() = 0;
 
-	/** Gets the extender for managing Navigation Tool customizations */
-	virtual UE::Sequencer::NavigationTool::FNavigationToolExtender& GetNavigationToolExtender() = 0;
-
 private:
 
 	/** Map of sequencer interfaces for movie scene channel types, keyed on channel UStruct name */
diff --git a/Engine/Source/Editor/Sequencer/Public/SequencerCommands.h b/Engine/Source/Editor/Sequencer/Public/SequencerCommands.h
--- a/Engine/Source/Editor/Sequencer/Public/SequencerCommands.h
+++ b/Engine/Source/Editor/Sequencer/Public/SequencerCommands.h
@@ -496,9 +496,6 @@ public:
 	/** Toggle the soloed state of Sequencer track selection */
 	TSharedPtr< FUICommandInfo > ToggleTrackSelectionSolo;
 
-	/** Toggles whether or not the Navigation Tool should be visible. */
-	TSharedPtr< FUICommandInfo > ToggleNavigationToolVisible;
-
 	/**
 	 * Initialize commands
 	 */
diff --git a/Engine/Source/Editor/Sequencer/Sequencer.Build.cs b/Engine/Source/Editor/Sequencer/Sequencer.Build.cs
--- a/Engine/Source/Editor/Sequencer/Sequencer.Build.cs
+++ b/Engine/Source/Editor/Sequencer/Sequencer.Build.cs
@@ -22,7 +22,6 @@ public class Sequencer : ModuleRules
 				"Core", 
 				"CoreUObject", 
 				"CurveEditor",
-				"DeveloperSettings",
 				"InputCore",
 				"Engine", 
 				"Slate", 
@@ -45,7 +44,6 @@ public class Sequencer : ModuleRules
 				"PropertyEditor",
 				"ViewportInteraction",
 				"SerializedRecorderInterface",
-				"SourceControl",
 				"SubobjectDataInterface",
 				"ToolMenus",
 				"ToolWidgets",
@@ -69,7 +67,6 @@ public class Sequencer : ModuleRules
 			new string[] {
 				"LevelEditor",
 				"DesktopPlatform",
-				"WorkspaceMenuStructure"
 				}
 			);
 
@@ -86,7 +83,6 @@ public class Sequencer : ModuleRules
 		DynamicallyLoadedModuleNames.AddRange(
 			new string[] {
 				"LevelEditor",
-				"WorkspaceMenuStructure",
 				"MainFrame",
 				}
 			);
