[SequenceNavigator]

diff --git a/Engine/Plugins/Experimental/SequenceNavigator/Source/SequenceNavigator/Private/NavigationTool.cpp b/Engine/Plugins/Experimental/SequenceNavigator/Source/SequenceNavigator/Private/NavigationTool.cpp
--- a/Engine/Plugins/Experimental/SequenceNavigator/Source/SequenceNavigator/Private/NavigationTool.cpp
+++ b/Engine/Plugins/Experimental/SequenceNavigator/Source/SequenceNavigator/Private/NavigationTool.cpp
@@ -34,7 +34,9 @@
 
 #define LOCTEXT_NAMESPACE "NavigationTool"
 
-namespace UE::Sequencer::NavigationTool
+using namespace UE::Sequencer;
+
+namespace UE::SequenceNavigator
 {
 
 FNavigationTool::FNavigationTool(const TWeakPtr<ISequencer>& InWeakSequencer)
@@ -70,7 +72,7 @@ void FNavigationTool::Init()
 	const UNavigationToolSettings* const ToolSettings = GetDefault<UNavigationToolSettings>();
 
 	// We apply default views *after* all columns have been created for all providers
-	if (ensure(ToolSettings) && ToolSettings->ShouldApplyDefaultColumnView())
+	if (ToolSettings && ToolSettings->ShouldApplyDefaultColumnView())
 	{
 		TArray<FText> DefaultColumnViews;
 
@@ -128,7 +130,7 @@ void FNavigationTool::ForEachProvider(const TFunction<bool(const TSharedRef<FNav
 	}
 
 	UMovieSceneSequence* const RootSequence = Sequencer->GetRootMovieSceneSequence();
-	if (!IsValid(RootSequence))
+	if (!RootSequence)
 	{
 		return;
 	}
@@ -282,7 +284,7 @@ bool FNavigationTool::DoesGlobalFilterExist(const FName InFilterId)
 void FNavigationTool::RefreshGlobalFilters()
 {
 	const UNavigationToolSettings* const ToolSettings = GetDefault<UNavigationToolSettings>();
-	if (!ensure(IsValid(ToolSettings)))
+	if (!ToolSettings)
 	{
 		return;
 	}
@@ -346,8 +348,7 @@ TSharedPtr<INavigationToolView> FNavigationTool::GetToolView(const int32 InToolV
 
 bool FNavigationTool::IsObjectAllowedInTool(const UObject* const InObject) const
 {
-	const AActor* const Actor = Cast<AActor>(InObject);
-	if (IsValid(Actor))
+	if (const AActor* const Actor = Cast<AActor>(InObject))
 	{
 		// Do not show Actors that aren't editable or not meant to be listed in Navigation Tool
 		if (!Actor->IsEditable())
@@ -358,8 +359,7 @@ bool FNavigationTool::IsObjectAllowedInTool(const UObject* const InObject) const
 		return true;
 	}
 
-	const UActorComponent* const ActorComponent = Cast<UActorComponent>(InObject);
-	if (IsValid(ActorComponent))
+	if (const UActorComponent* const ActorComponent = Cast<UActorComponent>(InObject))
 	{
 		const bool bHideConstructionScriptComponents = GetDefault<UBlueprintEditorSettings>()->bHideConstructionScriptComponentsInDetailsView;
 		return !ActorComponent->IsVisualizationComponent()
@@ -502,7 +502,7 @@ FNavigationToolItemPtr FNavigationTool::FindItem(const FNavigationToolItemId& In
 	return nullptr;
 }
 
-TArray<FNavigationToolItemPtr> FNavigationTool::TryFindItems(const FViewModelPtr& InViewModel) const
+TArray<FNavigationToolItemPtr> FNavigationTool::TryFindItems(const UE::Sequencer::FViewModelPtr& InViewModel) const
 {
 	if (const TViewModelPtr<FTrackModel> TrackModel = InViewModel.ImplicitCast())
 	{
@@ -1344,6 +1344,6 @@ void FNavigationTool::OnProvidersChanged(const FName InToolId
 	OnTreeViewChanged();
 }
 
-}
+} // namespace UE::SequenceNavigator
 
 #undef LOCTEXT_NAMESPACE
diff --git a/Engine/Plugins/Experimental/SequenceNavigator/Source/SequenceNavigator/Private/NavigationTool.h b/Engine/Plugins/Experimental/SequenceNavigator/Source/SequenceNavigator/Private/NavigationTool.h
--- a/Engine/Plugins/Experimental/SequenceNavigator/Source/SequenceNavigator/Private/NavigationTool.h
+++ b/Engine/Plugins/Experimental/SequenceNavigator/Source/SequenceNavigator/Private/NavigationTool.h
@@ -21,7 +21,7 @@ class ISequencer;
 class UMovieSceneSequence;
 class USequencerSettings;
 
-namespace UE::Sequencer::NavigationTool
+namespace UE::SequenceNavigator
 {
 
 class FNavigationToolBuiltInFilter;
@@ -92,7 +92,7 @@ public:
 	virtual TSharedRef<FNavigationToolItem> GetTreeRoot() const override;
 
 	virtual FNavigationToolItemPtr FindItem(const FNavigationToolItemId& InItemId) const override;
-	virtual TArray<FNavigationToolItemPtr> TryFindItems(const FViewModelPtr& InViewModel) const override;
+	virtual TArray<FNavigationToolItemPtr> TryFindItems(const UE::Sequencer::FViewModelPtr& InViewModel) const override;
 
 	virtual void SetIgnoreNotify(const ENavigationToolIgnoreNotifyFlags InFlag, const bool bInIgnore) override;
 
@@ -256,4 +256,4 @@ private:
 	TArray<TSharedPtr<FNavigationToolBuiltInFilter>> GlobalFilters;
 };
 
-}
+} // namespace UE::SequenceNavigator
diff --git a/Engine/Plugins/Experimental/SequenceNavigator/Source/SequenceNavigator/Private/NavigationToolDefines.cpp b/Engine/Plugins/Experimental/SequenceNavigator/Source/SequenceNavigator/Private/NavigationToolDefines.cpp
--- a/Engine/Plugins/Experimental/SequenceNavigator/Source/SequenceNavigator/Private/NavigationToolDefines.cpp
+++ b/Engine/Plugins/Experimental/SequenceNavigator/Source/SequenceNavigator/Private/NavigationToolDefines.cpp
@@ -4,15 +4,17 @@
 #include "Menus/NavigationToolItemContextMenu.h"
 #include "Menus/NavigationToolToolbarMenu.h"
 
-namespace UE::Sequencer::NavigationTool
+namespace UE::SequenceNavigator
 {
-	FName GetToolBarMenuName()
-	{
-		return FNavigationToolToolbarMenu::GetMenuName();
-	}
 
-	FName GetItemContextMenuName()
-	{
-		return FNavigationToolItemContextMenu::GetMenuName();
-	}
+FName GetToolBarMenuName()
+{
+	return FNavigationToolToolbarMenu::GetMenuName();
 }
+
+FName GetItemContextMenuName()
+{
+	return FNavigationToolItemContextMenu::GetMenuName();
+}
+
+} // namespace UE::SequenceNavigator
diff --git a/Engine/Plugins/Experimental/SequenceNavigator/Source/SequenceNavigator/Private/NavigationToolExtender.cpp b/Engine/Plugins/Experimental/SequenceNavigator/Source/SequenceNavigator/Private/NavigationToolExtender.cpp
--- a/Engine/Plugins/Experimental/SequenceNavigator/Source/SequenceNavigator/Private/NavigationToolExtender.cpp
+++ b/Engine/Plugins/Experimental/SequenceNavigator/Source/SequenceNavigator/Private/NavigationToolExtender.cpp
@@ -14,7 +14,7 @@
 #include "SequenceNavigatorLog.h"
 #include "SequencerSettings.h"
 
-namespace UE::Sequencer::NavigationTool
+namespace UE::SequenceNavigator
 {
 
 FNavigationToolExtender& FNavigationToolExtender::Get()
@@ -455,4 +455,4 @@ TSharedPtr<INavigationToolIconCustomization> FNavigationToolExtender::GetCustomi
 	return nullptr;
 }
 
-} // namespace UE::Sequencer::NavigationTool
+} // namespace UE::SequenceNavigator
diff --git a/Engine/Plugins/Experimental/SequenceNavigator/Source/SequenceNavigator/Private/NavigationToolScopedSelection.cpp b/Engine/Plugins/Experimental/SequenceNavigator/Source/SequenceNavigator/Private/NavigationToolScopedSelection.cpp
--- a/Engine/Plugins/Experimental/SequenceNavigator/Source/SequenceNavigator/Private/NavigationToolScopedSelection.cpp
+++ b/Engine/Plugins/Experimental/SequenceNavigator/Source/SequenceNavigator/Private/NavigationToolScopedSelection.cpp
@@ -14,10 +14,11 @@
 #include "NavigationToolSettings.h"
 #include "Utils/NavigationToolMiscUtils.h"
 
-namespace UE::Sequencer::NavigationTool
-{
 using namespace UE::Sequencer;
 
+namespace UE::SequenceNavigator
+{
+
 FNavigationToolScopedSelection::FNavigationToolScopedSelection(ISequencer& InSequencer
 	, ENavigationToolScopedSelectionPurpose InPurpose)
 	: Sequencer(InSequencer)
@@ -66,7 +67,7 @@ void FNavigationToolScopedSelection::Select(UMovieSceneSection* const InSection)
 		return;
 	}
 
-	if (!IsValid(InSection))
+	if (!InSection)
 	{
 		return;
 	}
@@ -89,7 +90,7 @@ void FNavigationToolScopedSelection::Select(UMovieSceneTrack* const InTrack)
 		return;
 	}
 
-	if (!IsValid(InTrack))
+	if (!InTrack)
 	{
 		return;
 	}
@@ -113,7 +114,7 @@ void FNavigationToolScopedSelection::Select(UMovieSceneSequence* const InSequenc
 		return;
 	}
 
-	if (!IsValid(InSequence) || InMarkedFrameIndex == INDEX_NONE)
+	if (!InSequence || InMarkedFrameIndex == INDEX_NONE)
 	{
 		return;
 	}
@@ -144,7 +145,7 @@ bool FNavigationToolScopedSelection::IsSelected(UMovieSceneTrack* const InTrack)
 bool FNavigationToolScopedSelection::IsSelected(UMovieSceneSequence* const InSequence
 	, const int32 InMarkedFrameIndex) const
 {
-	if (IsValid(InSequence) && SelectedMarkedFrames.Contains(InSequence))
+	if (InSequence && SelectedMarkedFrames.Contains(InSequence))
 	{
 		return SelectedMarkedFrames[InSequence].Contains(InMarkedFrameIndex);
 	}
@@ -246,4 +247,4 @@ ISequencer& FNavigationToolScopedSelection::GetSequencer() const
 	return Sequencer;
 }
 
-}
+} // namespace UE::SequenceNavigator
diff --git a/Engine/Plugins/Experimental/SequenceNavigator/Source/SequenceNavigator/Private/NavigationToolSettings.cpp b/Engine/Plugins/Experimental/SequenceNavigator/Source/SequenceNavigator/Private/NavigationToolSettings.cpp
--- a/Engine/Plugins/Experimental/SequenceNavigator/Source/SequenceNavigator/Private/NavigationToolSettings.cpp
+++ b/Engine/Plugins/Experimental/SequenceNavigator/Source/SequenceNavigator/Private/NavigationToolSettings.cpp
@@ -13,8 +13,7 @@ void UNavigationToolSettings::OpenEditorSettings()
 {
 	if (ISettingsModule* const SettingsModule = FModuleManager::GetModulePtr<ISettingsModule>(TEXT("Settings")))
 	{
-		const UNavigationToolSettings* const Settings = GetDefault<UNavigationToolSettings>();
-		if (ensure(IsValid(Settings)))
+		if (const UNavigationToolSettings* const Settings = GetDefault<UNavigationToolSettings>())
 		{
 			SettingsModule->ShowViewer(Settings->GetContainerName(), Settings->GetCategoryName(), Settings->GetSectionName());
 		}
diff --git a/Engine/Plugins/Experimental/SequenceNavigator/Source/SequenceNavigator/Private/NavigationToolTab.cpp b/Engine/Plugins/Experimental/SequenceNavigator/Source/SequenceNavigator/Private/NavigationToolTab.cpp
--- a/Engine/Plugins/Experimental/SequenceNavigator/Source/SequenceNavigator/Private/NavigationToolTab.cpp
+++ b/Engine/Plugins/Experimental/SequenceNavigator/Source/SequenceNavigator/Private/NavigationToolTab.cpp
@@ -13,7 +13,7 @@
 
 #define LOCTEXT_NAMESPACE "NavigationToolTab"
 
-namespace UE::Sequencer::NavigationTool
+namespace UE::SequenceNavigator
 {
 
 FNavigationToolTab::FNavigationToolTab(INavigationTool& InOwnerTool)
@@ -28,7 +28,7 @@ void FNavigationToolTab::Init()
 	RegisterToolTab();
 
 	USequencerSettings* const SequencerSettings = GetSequencerSettings();
-    if (ensure(IsValid(SequencerSettings)) && SequencerSettings->IsNavigationToolVisible())
+    if (SequencerSettings && SequencerSettings->IsNavigationToolVisible())
     {
     	if (const TSharedPtr<FTabManager> TabManager = GetTabManager())
     	{
@@ -47,8 +47,7 @@ void FNavigationToolTab::Shutdown()
 
 bool FNavigationToolTab::ShouldShowToolTab() const
 {
-	USequencerSettings* const SequencerSettings = GetSequencerSettings();
-	if (ensure(IsValid(SequencerSettings)))
+	if (USequencerSettings* const SequencerSettings = GetSequencerSettings())
 	{
 		return SequencerSettings->IsNavigationToolVisible();
 	}
@@ -89,8 +88,7 @@ FName FNavigationToolTab::StaticToolTabId(const ISequencer& InSequencer)
 {
 	const FString DefaultTabId = TEXT("NavigationTool");
 
-	const USequencerSettings* const SequencerSettings = InSequencer.GetSequencerSettings();
-	if (ensure(IsValid(SequencerSettings)))
+	if (const USequencerSettings* const SequencerSettings = InSequencer.GetSequencerSettings())
 	{
 		const FString TabName = DefaultTabId + TEXT(".") + SequencerSettings->GetName();
 		return *TabName;
@@ -219,6 +217,6 @@ void FNavigationToolTab::CloseToolTab()
 	LiveTab->RequestCloseTab();
 }
 
-}
+} // namespace UE::SequenceNavigator
 
 #undef LOCTEXT_NAMESPACE
diff --git a/Engine/Plugins/Experimental/SequenceNavigator/Source/SequenceNavigator/Private/Providers/NavigationToolProvider.cpp b/Engine/Plugins/Experimental/SequenceNavigator/Source/SequenceNavigator/Private/Providers/NavigationToolProvider.cpp
--- a/Engine/Plugins/Experimental/SequenceNavigator/Source/SequenceNavigator/Private/Providers/NavigationToolProvider.cpp
+++ b/Engine/Plugins/Experimental/SequenceNavigator/Source/SequenceNavigator/Private/Providers/NavigationToolProvider.cpp
@@ -8,7 +8,7 @@
 #include "NavigationToolView.h"
 #include "SequenceNavigatorLog.h"
 
-namespace UE::Sequencer::NavigationTool
+namespace UE::SequenceNavigator
 {
 
 void FNavigationToolProvider::OnExtendColumnViews(TSet<FNavigationToolColumnView>& OutColumnViews)
@@ -167,7 +167,7 @@ void FNavigationToolProvider::Deactivate(FNavigationTool& InTool)
 void FNavigationToolProvider::CleanupExtendedColumnViews()
 {
 	UNavigationToolSettings* const Settings = GetMutableDefault<UNavigationToolSettings>();
-	if (!ensure(IsValid(Settings)))
+	if (!Settings)
 	{
 		return;
 	}
@@ -184,7 +184,7 @@ void FNavigationToolProvider::CleanupExtendedColumnViews()
 
 bool FNavigationToolProvider::IsSequenceSupported(UMovieSceneSequence* const InSequence) const
 {
-	if (IsValid(InSequence))
+	if (InSequence)
 	{
 		return GetSupportedSequenceClasses().Contains(InSequence->GetClass());
 	}
@@ -222,7 +222,7 @@ void FNavigationToolProvider::SaveState(FNavigationTool& InTool)
 void FNavigationToolProvider::LoadState(FNavigationTool& InTool)
 {
 	const UNavigationToolSettings* const ToolSettings = GetDefault<UNavigationToolSettings>();
-	if (!ensure(IsValid(ToolSettings)))
+	if (!ToolSettings)
 	{
 		return;
 	}
@@ -312,4 +312,4 @@ void FNavigationToolProvider::LoadSerializedTree(const FNavigationToolItemPtr& I
 	}
 }
 
-}
+} // namespace UE::SequenceNavigator
diff --git a/Engine/Plugins/Experimental/SequenceNavigator/Source/SequenceNavigator/Private/Utils/NavigationToolMiscUtils.cpp b/Engine/Plugins/Experimental/SequenceNavigator/Source/SequenceNavigator/Private/Utils/NavigationToolMiscUtils.cpp
--- a/Engine/Plugins/Experimental/SequenceNavigator/Source/SequenceNavigator/Private/Utils/NavigationToolMiscUtils.cpp
+++ b/Engine/Plugins/Experimental/SequenceNavigator/Source/SequenceNavigator/Private/Utils/NavigationToolMiscUtils.cpp
@@ -16,12 +16,12 @@
 #include "Sections/MovieSceneSubSection.h"
 #include "UObject/Package.h"
 
-namespace UE::Sequencer::NavigationTool
+namespace UE::SequenceNavigator
 {
 
 FSourceControlStatePtr FindSourceControlState(UPackage* const InPackage)
 {
-	if (!IsValid(InPackage))
+	if (!InPackage)
 	{
 		return nullptr;
 	}
@@ -90,7 +90,7 @@ void FocusSequence(const INavigationTool& InTool, UMovieSceneSequence& InSequenc
 
 	FMovieSceneSequenceIDRef SequenceID = EvaluationState->FindSequenceId(&InSequence);
 	UMovieSceneSubSection* const SubSection = Sequencer->FindSubSection(SequenceID);
-	if (!IsValid(SubSection))
+	if (!SubSection)
 	{
 		return;
 	}
@@ -106,8 +106,7 @@ void FocusSequence(const INavigationTool& InTool
 
 	if (const TSharedPtr<ISequencer> Sequencer = InTool.GetSequencer())
 	{
-		UMovieSceneSubSection* const SubSection = InSequenceItem.GetSubSection();
-		if (IsValid(SubSection))
+		if (UMovieSceneSubSection* const SubSection = InSequenceItem.GetSubSection())
 		{
 			Sequencer->EmptySelection();
 			Sequencer->SelectSection(SubSection);
@@ -119,24 +118,25 @@ void FocusSequence(const INavigationTool& InTool
 	, UMovieSceneSequence& InSequence
 	, const FNavigationToolTrack& InTrackItem)
 {
+	using namespace UE::Sequencer;
+
 	FocusSequence(InTool, InSequence);
 
 	if (const TSharedPtr<ISequencer> Sequencer = InTool.GetSequencer())
 	{
-		UMovieSceneTrack* const Track = InTrackItem.GetTrack();
-		if (IsValid(Track))
+		if (UMovieSceneTrack* const Track = InTrackItem.GetTrack())
 		{
 			Sequencer->EmptySelection();
 			Sequencer->SelectTrack(Track);
 
-			/*if (const TSharedPtr<SOutlinerView> OutlinerView = Sequencer->GetOutlinerViewWidget())
+			if (const TSharedPtr<SOutlinerView> OutlinerView = Sequencer->GetOutlinerViewWidget())
 			{
 				if (const TViewModelPtr<FTrackModel> TrackModel = InTrackItem.GetViewModel())
 				{
 					OutlinerView->ExpandCollapseNode(TrackModel, true, ETreeRecursion::NonRecursive);
 					OutlinerView->RequestScrollIntoView(TrackModel);
 				}
-			}*/
+			}
 		}
 	}
 }
@@ -145,25 +145,29 @@ void FocusSequence(const INavigationTool& InTool
 	, UMovieSceneSequence& InSequence
 	, const FNavigationToolBinding& InBindingItem)
 {
+	using namespace UE::Sequencer;
+
 	FocusSequence(InTool, InSequence);
 
 	if (const TSharedPtr<ISequencer> Sequencer = InTool.GetSequencer())
 	{
 		SelectSequencerBindingTrack(*Sequencer, InBindingItem.GetBinding().GetObjectGuid());
 
-		/*if (const TSharedPtr<SOutlinerView> OutlinerView = Sequencer->GetOutlinerViewWidget())
+		if (const TSharedPtr<SOutlinerView> OutlinerView = Sequencer->GetOutlinerViewWidget())
 		{
 			if (const TViewModelPtr<FObjectBindingModel> ObjectBindingViewModel = InBindingItem.GetViewModel())
 			{
 				OutlinerView->ExpandCollapseNode(ObjectBindingViewModel, true, ETreeRecursion::NonRecursive);
 				OutlinerView->RequestScrollIntoView(ObjectBindingViewModel);
 			}
-		}*/
+		}
 	}
 }
 
 void SelectSequencerBindingTrack(const ISequencer& InSequencer, const FGuid& InObjectId)
 {
+	using namespace UE::Sequencer;
+
 	const TSharedPtr<FSequencerEditorViewModel> ViewModel = InSequencer.GetViewModel();
 	if (!ViewModel.IsValid())
 	{
@@ -261,4 +265,4 @@ TArrayView<TWeakObjectPtr<>> ResolveBoundObjects(const ISequencer& InSequencer
 	return TArrayView<TWeakObjectPtr<>>();
 }
 
-}
+} // namespace UE::SequenceNavigator
diff --git a/Engine/Plugins/Experimental/SequenceNavigator/Source/SequenceNavigator/Private/Utils/NavigationToolMiscUtils.h b/Engine/Plugins/Experimental/SequenceNavigator/Source/SequenceNavigator/Private/Utils/NavigationToolMiscUtils.h
--- a/Engine/Plugins/Experimental/SequenceNavigator/Source/SequenceNavigator/Private/Utils/NavigationToolMiscUtils.h
+++ b/Engine/Plugins/Experimental/SequenceNavigator/Source/SequenceNavigator/Private/Utils/NavigationToolMiscUtils.h
@@ -17,12 +17,7 @@ struct FGuid;
 struct FMovieSceneSequenceID;
 struct FSlateBrush;
 
-namespace UE::Sequencer
-{
-class SOutlinerView;
-}
-
-namespace UE::Sequencer::NavigationTool
+namespace UE::SequenceNavigator
 {
 
 class INavigationTool;
@@ -119,4 +114,4 @@ FMovieSceneSequenceID ResolveSequenceID(const ISequencer& InSequencer, UMovieSce
 TArrayView<TWeakObjectPtr<>> ResolveBoundObjects(const ISequencer& InSequencer
 	, UMovieSceneSequence* const InSequence, const FGuid& InBindingId);
 
-}
+} // namespace UE::SequenceNavigator
diff --git a/Engine/Plugins/Experimental/SequenceNavigator/Source/SequenceNavigator/Private/Utils/NavigationToolMovieSceneUtils.cpp b/Engine/Plugins/Experimental/SequenceNavigator/Source/SequenceNavigator/Private/Utils/NavigationToolMovieSceneUtils.cpp
--- a/Engine/Plugins/Experimental/SequenceNavigator/Source/SequenceNavigator/Private/Utils/NavigationToolMovieSceneUtils.cpp
+++ b/Engine/Plugins/Experimental/SequenceNavigator/Source/SequenceNavigator/Private/Utils/NavigationToolMovieSceneUtils.cpp
@@ -7,13 +7,12 @@
 #include "Sections/MovieSceneSubSection.h"
 #include "SequencerUtilities.h"
 
-namespace UE::Sequencer::NavigationTool
+namespace UE::SequenceNavigator
 {
 
 FFrameTime ConvertToDisplayRateTime(const UMovieSceneSequence& InSequence, const FFrameTime& InTime)
 {
-	UMovieScene* const MovieScene = InSequence.GetMovieScene();
-	if (IsValid(MovieScene))
+	if (UMovieScene* const MovieScene = InSequence.GetMovieScene())
 	{
 		return ConvertFrameTime(InTime
 			, MovieScene->GetTickResolution()
@@ -24,8 +23,7 @@ FFrameTime ConvertToDisplayRateTime(const UMovieSceneSequence& InSequence, const
 
 FFrameTime ConvertToTickResolutionTime(const UMovieSceneSequence& InSequence, const FFrameTime& InTime)
 {
-	UMovieScene* const MovieScene = InSequence.GetMovieScene();
-	if (IsValid(MovieScene))
+	if (UMovieScene* const MovieScene = InSequence.GetMovieScene())
 	{
 		return ConvertFrameTime(InTime
 			, MovieScene->GetDisplayRate()
@@ -49,13 +47,13 @@ UMovieSceneSubSection* FindSequenceSubSection(ISequencer& InSequencer, UMovieSce
 
 bool IsGloballyMarkedFramesForSequence(UMovieSceneSequence* const InSequence)
 {
-	if (!IsValid(InSequence))
+	if (!InSequence)
 	{
 		return false;
 	}
 
 	const UMovieScene* const MovieScene = InSequence->GetMovieScene();
-	if (!IsValid(MovieScene))
+	if (!MovieScene)
 	{
 		return false;
 	}
@@ -65,13 +63,13 @@ bool IsGloballyMarkedFramesForSequence(UMovieSceneSequence* const InSequence)
 
 void ShowGloballyMarkedFramesForSequence(ISequencer& InSequencer, UMovieSceneSequence* const InSequence, const bool bInVisible)
 {
-	if (!IsValid(InSequence))
+	if (!InSequence)
 	{
 		return;
 	}
 
 	UMovieScene* const MovieScene = InSequence->GetMovieScene();
-	if (!IsValid(MovieScene))
+	if (!MovieScene)
 	{
 		return;
 	}
@@ -92,13 +90,13 @@ void ShowGloballyMarkedFramesForSequence(ISequencer& InSequencer, UMovieSceneSeq
 
 void ModifySequenceAndMovieScene(UMovieSceneSequence* const InSequence)
 {
-	if (!IsValid(InSequence))
+	if (!InSequence)
 	{
 		return;
 	}
 
 	UMovieScene* const MovieScene = InSequence->GetMovieScene();
-	if (!IsValid(MovieScene))
+	if (!MovieScene)
 	{
 		return;
 	}
@@ -109,25 +107,24 @@ void ModifySequenceAndMovieScene(UMovieSceneSequence* const InSequence)
 
 void GetSequenceSubSections(UMovieSceneSequence* const InSequence, TArray<UMovieSceneSubSection*>& OutSubSections)
 {
-	if (!IsValid(InSequence))
+	if (!InSequence)
 	{
 		return;
 	}
 
 	const UMovieScene* const MovieScene = InSequence->GetMovieScene();
-	if (!IsValid(MovieScene))
+	if (!MovieScene)
 	{
 		return;
 	}
 
 	for (UMovieSceneSection* const Section : MovieScene->GetAllSections())
 	{
-		UMovieSceneSubSection* const SubSection = Cast<UMovieSceneSubSection>(Section);
-		if (IsValid(SubSection))
+		if (UMovieSceneSubSection* const SubSection = Cast<UMovieSceneSubSection>(Section))
 		{
 			OutSubSections.Add(SubSection);
 		}
 	}
 }
 
-}
+} // namespace UE::SequenceNavigator
diff --git a/Engine/Plugins/Experimental/SequenceNavigator/Source/SequenceNavigator/Private/Widgets/Columns/SNavigationToolEndFrameOffset.cpp b/Engine/Plugins/Experimental/SequenceNavigator/Source/SequenceNavigator/Private/Widgets/Columns/SNavigationToolEndFrameOffset.cpp
--- a/Engine/Plugins/Experimental/SequenceNavigator/Source/SequenceNavigator/Private/Widgets/Columns/SNavigationToolEndFrameOffset.cpp
+++ b/Engine/Plugins/Experimental/SequenceNavigator/Source/SequenceNavigator/Private/Widgets/Columns/SNavigationToolEndFrameOffset.cpp
@@ -7,10 +7,11 @@
 #include "Sections/MovieSceneSubSection.h"
 #include "Widgets/SNavigationToolTreeRow.h"
 
-using namespace UE::Sequencer::NavigationTool;
-
 #define LOCTEXT_NAMESPACE "SNavigationToolEndFrameOffset"
 
+namespace UE::SequenceNavigator
+{
+
 void SNavigationToolEndFrameOffset::Construct(const FArguments& InArgs
 	, const FNavigationToolItemRef& InItem
 	, const TSharedRef<INavigationToolView>& InView
@@ -64,4 +65,6 @@ FText SNavigationToolEndFrameOffset::GetTransactionText() const
 	return LOCTEXT("SetEndFrameOffsetTransaction", "Set End Frame Offset");
 }
 
+} // namespace UE::SequenceNavigator
+
 #undef LOCTEXT_NAMESPACE
diff --git a/Engine/Plugins/Experimental/SequenceNavigator/Source/SequenceNavigator/Private/Widgets/Columns/SNavigationToolLength.cpp b/Engine/Plugins/Experimental/SequenceNavigator/Source/SequenceNavigator/Private/Widgets/Columns/SNavigationToolLength.cpp
--- a/Engine/Plugins/Experimental/SequenceNavigator/Source/SequenceNavigator/Private/Widgets/Columns/SNavigationToolLength.cpp
+++ b/Engine/Plugins/Experimental/SequenceNavigator/Source/SequenceNavigator/Private/Widgets/Columns/SNavigationToolLength.cpp
@@ -10,7 +10,7 @@
 
 #define LOCTEXT_NAMESPACE "SNavigationToolLength"
 
-namespace UE::Sequencer::NavigationTool
+namespace UE::SequenceNavigator
 {
 
 void SNavigationToolLength::Construct(const FArguments& InArgs
@@ -50,11 +50,11 @@ TSharedPtr<INumericTypeInterface<double>> SNavigationToolLength::GetNumericTypeI
 {
 	if (const TSharedPtr<ISequencer> Sequencer = GetSequencer())
 	{
-		return Sequencer->GetNumericTypeInterface(ENumericIntent::Duration);
+		return Sequencer->GetNumericTypeInterface(UE::Sequencer::ENumericIntent::Duration);
 	}
 	return nullptr;
 }
 
-}
+} // namespace UE::SequenceNavigator
 
 #undef LOCTEXT_NAMESPACE
diff --git a/Engine/Plugins/Experimental/SequenceNavigator/Source/SequenceNavigator/Private/Widgets/Columns/SNavigationToolTake.cpp b/Engine/Plugins/Experimental/SequenceNavigator/Source/SequenceNavigator/Private/Widgets/Columns/SNavigationToolTake.cpp
--- a/Engine/Plugins/Experimental/SequenceNavigator/Source/SequenceNavigator/Private/Widgets/Columns/SNavigationToolTake.cpp
+++ b/Engine/Plugins/Experimental/SequenceNavigator/Source/SequenceNavigator/Private/Widgets/Columns/SNavigationToolTake.cpp
@@ -14,11 +14,9 @@
 #include "Widgets/Layout/SBox.h"
 #include "Widgets/SNavigationToolTreeRow.h"
 
-using namespace UE::Sequencer::NavigationTool;
-
 #define LOCTEXT_NAMESPACE "SNavigationToolTake"
 
-namespace UE::Sequencer::NavigationTool
+namespace UE::SequenceNavigator
 {
 
 DECLARE_DELEGATE_RetVal_OneParam(FReply, FOnTakeEntrySelected, const TSharedPtr<SNavigationToolTake::FTakeItemInfo> /*InTakeInfo*/);
@@ -389,6 +387,6 @@ void SNavigationToolTake::SetActiveTake(UMovieSceneSequence* const InSequence)
 	}
 }
 
-}
+} // namespace UE::SequenceNavigator
 
 #undef LOCTEXT_NAMESPACE
diff --git a/Engine/Plugins/Experimental/SequenceNavigator/Source/SequenceNavigator/Private/Widgets/Columns/SNavigationToolTime.cpp b/Engine/Plugins/Experimental/SequenceNavigator/Source/SequenceNavigator/Private/Widgets/Columns/SNavigationToolTime.cpp
--- a/Engine/Plugins/Experimental/SequenceNavigator/Source/SequenceNavigator/Private/Widgets/Columns/SNavigationToolTime.cpp
+++ b/Engine/Plugins/Experimental/SequenceNavigator/Source/SequenceNavigator/Private/Widgets/Columns/SNavigationToolTime.cpp
@@ -11,7 +11,7 @@
 
 #define LOCTEXT_NAMESPACE "SNavigationToolTime"
 
-namespace UE::Sequencer::NavigationTool
+namespace UE::SequenceNavigator
 {
 
 void SNavigationToolTime::Construct(const FArguments& InArgs
@@ -82,7 +82,7 @@ TSharedPtr<INumericTypeInterface<double>> SNavigationToolTime::GetNumericTypeInt
 {
 	if (const TSharedPtr<ISequencer> Sequencer = GetSequencer())
 	{
-		return Sequencer->GetNumericTypeInterface(ENumericIntent::Position);
+		return Sequencer->GetNumericTypeInterface(UE::Sequencer::ENumericIntent::Position);
 	}
 	return nullptr;
 }
@@ -105,6 +105,6 @@ TSharedPtr<ISequencer> SNavigationToolTime::GetSequencer() const
 	return nullptr;
 }
 
-}
+} // namespace UE::SequenceNavigator
 
 #undef LOCTEXT_NAMESPACE
diff --git a/Engine/Plugins/Experimental/SequenceNavigator/Source/SequenceNavigator/Private/Widgets/Columns/SNavigationToolTime.h b/Engine/Plugins/Experimental/SequenceNavigator/Source/SequenceNavigator/Private/Widgets/Columns/SNavigationToolTime.h
--- a/Engine/Plugins/Experimental/SequenceNavigator/Source/SequenceNavigator/Private/Widgets/Columns/SNavigationToolTime.h
+++ b/Engine/Plugins/Experimental/SequenceNavigator/Source/SequenceNavigator/Private/Widgets/Columns/SNavigationToolTime.h
@@ -8,9 +8,11 @@
 #include "Widgets/Input/NumericTypeInterface.h"
 #include "Widgets/SCompoundWidget.h"
 
+class FName;
+class FText;
 class ISequencer;
 
-namespace UE::Sequencer::NavigationTool
+namespace UE::SequenceNavigator
 {
 
 class INavigationToolView;
@@ -62,4 +64,4 @@ protected:
 	bool bIsReadOnly = false;
 };
 
-}
+} // namespace UE::SequenceNavigator
diff --git a/Engine/Plugins/Experimental/SequenceNavigator/Source/SequenceNavigator/Public/Extensions/IColorExtension.h b/Engine/Plugins/Experimental/SequenceNavigator/Source/SequenceNavigator/Public/Extensions/IColorExtension.h
--- a/Engine/Plugins/Experimental/SequenceNavigator/Source/SequenceNavigator/Public/Extensions/IColorExtension.h
+++ b/Engine/Plugins/Experimental/SequenceNavigator/Source/SequenceNavigator/Public/Extensions/IColorExtension.h
@@ -5,17 +5,17 @@
 #include "Misc/Optional.h"
 #include "NavigationToolItemType.h"
 
-namespace UE::Sequencer::NavigationTool
+namespace UE::SequenceNavigator
 {
 
-class SEQUENCENAVIGATOR_API IColorExtension : public INavigationToolItemTypeCastable
+class IColorExtension : public INavigationToolItemTypeCastable
 {
 public:
 	UE_NAVIGATIONTOOL_INHERITS(IColorExtension, INavigationToolItemTypeCastable)
 
-	virtual TOptional<FColor> GetColor() const = 0;
+	SEQUENCENAVIGATOR_API virtual TOptional<FColor> GetColor() const = 0;
 
-	virtual void SetColor(const TOptional<FColor>& InColor) = 0;
+	SEQUENCENAVIGATOR_API virtual void SetColor(const TOptional<FColor>& InColor) = 0;
 };
 
-}
+} // namespace UE::SequenceNavigator
diff --git a/Engine/Plugins/Experimental/SequenceNavigator/Source/SequenceNavigator/Public/Extensions/IInTimeExtension.h b/Engine/Plugins/Experimental/SequenceNavigator/Source/SequenceNavigator/Public/Extensions/IInTimeExtension.h
--- a/Engine/Plugins/Experimental/SequenceNavigator/Source/SequenceNavigator/Public/Extensions/IInTimeExtension.h
+++ b/Engine/Plugins/Experimental/SequenceNavigator/Source/SequenceNavigator/Public/Extensions/IInTimeExtension.h
@@ -5,17 +5,17 @@
 #include "Misc/FrameNumber.h"
 #include "NavigationToolItemType.h"
 
-namespace UE::Sequencer::NavigationTool
+namespace UE::SequenceNavigator
 {
 
-class SEQUENCENAVIGATOR_API IInTimeExtension : public INavigationToolItemTypeCastable
+class IInTimeExtension : public INavigationToolItemTypeCastable
 {
 public:
 	UE_NAVIGATIONTOOL_INHERITS(IInTimeExtension, INavigationToolItemTypeCastable)
 
-	virtual FFrameNumber GetInTime() const = 0;
+	SEQUENCENAVIGATOR_API virtual FFrameNumber GetInTime() const = 0;
 
-	virtual void SetInTime(const FFrameNumber& InTime) = 0;
+	SEQUENCENAVIGATOR_API virtual void SetInTime(const FFrameNumber& InTime) = 0;
 };
 
-}
+} // namespace UE::SequenceNavigator
diff --git a/Engine/Plugins/Experimental/SequenceNavigator/Source/SequenceNavigator/Public/Extensions/IMarkerVisibilityExtension.h b/Engine/Plugins/Experimental/SequenceNavigator/Source/SequenceNavigator/Public/Extensions/IMarkerVisibilityExtension.h
--- a/Engine/Plugins/Experimental/SequenceNavigator/Source/SequenceNavigator/Public/Extensions/IMarkerVisibilityExtension.h
+++ b/Engine/Plugins/Experimental/SequenceNavigator/Source/SequenceNavigator/Public/Extensions/IMarkerVisibilityExtension.h
@@ -4,7 +4,7 @@
 
 #include "NavigationToolItemType.h"
 
-namespace UE::Sequencer::NavigationTool
+namespace UE::SequenceNavigator
 {
 
 enum class EItemMarkerVisibility
@@ -15,14 +15,14 @@ enum class EItemMarkerVisibility
 };
 ENUM_CLASS_FLAGS(EItemMarkerVisibility)
 
-class SEQUENCENAVIGATOR_API IMarkerVisibilityExtension : public INavigationToolItemTypeCastable
+class IMarkerVisibilityExtension : public INavigationToolItemTypeCastable
 {
 public:
 	UE_NAVIGATIONTOOL_INHERITS(IMarkerVisibilityExtension, INavigationToolItemTypeCastable)
 
-	virtual EItemMarkerVisibility GetMarkerVisibility() const = 0;
+	SEQUENCENAVIGATOR_API virtual EItemMarkerVisibility GetMarkerVisibility() const = 0;
 
-	virtual void SetMarkerVisibility(const bool bInVisible) = 0;
+	SEQUENCENAVIGATOR_API virtual void SetMarkerVisibility(const bool bInVisible) = 0;
 };
 
-}
+} // namespace UE::SequenceNavigator
diff --git a/Engine/Plugins/Experimental/SequenceNavigator/Source/SequenceNavigator/Public/Extensions/IOutTimeExtension.h b/Engine/Plugins/Experimental/SequenceNavigator/Source/SequenceNavigator/Public/Extensions/IOutTimeExtension.h
--- a/Engine/Plugins/Experimental/SequenceNavigator/Source/SequenceNavigator/Public/Extensions/IOutTimeExtension.h
+++ b/Engine/Plugins/Experimental/SequenceNavigator/Source/SequenceNavigator/Public/Extensions/IOutTimeExtension.h
@@ -5,17 +5,17 @@
 #include "Misc/FrameNumber.h"
 #include "NavigationToolItemType.h"
 
-namespace UE::Sequencer::NavigationTool
+namespace UE::SequenceNavigator
 {
 
-class SEQUENCENAVIGATOR_API IOutTimeExtension : public INavigationToolItemTypeCastable
+class IOutTimeExtension : public INavigationToolItemTypeCastable
 {
 public:
 	UE_NAVIGATIONTOOL_INHERITS(IOutTimeExtension, INavigationToolItemTypeCastable)
 
-	virtual FFrameNumber GetOutTime() const = 0;
+	SEQUENCENAVIGATOR_API virtual FFrameNumber GetOutTime() const = 0;
 
-	virtual void SetOutTime(const FFrameNumber& InTime) = 0;
+	SEQUENCENAVIGATOR_API virtual void SetOutTime(const FFrameNumber& InTime) = 0;
 };
 
-}
+} // namespace UE::SequenceNavigator
diff --git a/Engine/Plugins/Experimental/SequenceNavigator/Source/SequenceNavigator/Public/Extensions/IPlayheadExtension.h b/Engine/Plugins/Experimental/SequenceNavigator/Source/SequenceNavigator/Public/Extensions/IPlayheadExtension.h
--- a/Engine/Plugins/Experimental/SequenceNavigator/Source/SequenceNavigator/Public/Extensions/IPlayheadExtension.h
+++ b/Engine/Plugins/Experimental/SequenceNavigator/Source/SequenceNavigator/Public/Extensions/IPlayheadExtension.h
@@ -4,7 +4,7 @@
 
 #include "NavigationToolItemType.h"
 
-namespace UE::Sequencer::NavigationTool
+namespace UE::SequenceNavigator
 {
 
 enum class EItemContainsPlayhead
@@ -15,12 +15,12 @@ enum class EItemContainsPlayhead
 };
 ENUM_CLASS_FLAGS(EItemContainsPlayhead)
 
-class SEQUENCENAVIGATOR_API IPlayheadExtension : public INavigationToolItemTypeCastable
+class IPlayheadExtension : public INavigationToolItemTypeCastable
 {
 public:
 	UE_NAVIGATIONTOOL_INHERITS(IPlayheadExtension, INavigationToolItemTypeCastable)
 
-	virtual EItemContainsPlayhead ContainsPlayhead() const = 0;
+	SEQUENCENAVIGATOR_API virtual EItemContainsPlayhead ContainsPlayhead() const = 0;
 };
 
-}
+} // namespace UE::SequenceNavigator
diff --git a/Engine/Plugins/Experimental/SequenceNavigator/Source/SequenceNavigator/Public/Extensions/IRenameableExtension.h b/Engine/Plugins/Experimental/SequenceNavigator/Source/SequenceNavigator/Public/Extensions/IRenameableExtension.h
--- a/Engine/Plugins/Experimental/SequenceNavigator/Source/SequenceNavigator/Public/Extensions/IRenameableExtension.h
+++ b/Engine/Plugins/Experimental/SequenceNavigator/Source/SequenceNavigator/Public/Extensions/IRenameableExtension.h
@@ -4,19 +4,19 @@
 
 #include "NavigationToolItemType.h"
 
-namespace UE::Sequencer::NavigationTool
+namespace UE::SequenceNavigator
 {
 
-class SEQUENCENAVIGATOR_API IRenameableExtension : public INavigationToolItemTypeCastable
+class IRenameableExtension : public INavigationToolItemTypeCastable
 {
 public:
 	UE_NAVIGATIONTOOL_INHERITS(IRenameableExtension, INavigationToolItemTypeCastable)
 
 	/** @return True if this item can be renamed */
-	virtual bool CanRename() const = 0;
+	SEQUENCENAVIGATOR_API virtual bool CanRename() const = 0;
 
 	/** The implementation to rename the item */
-	virtual bool Rename(const FString& InName) = 0;
+	SEQUENCENAVIGATOR_API virtual bool Rename(const FString& InName) = 0;
 };
 
-}
+} // namespace UE::SequenceNavigator
diff --git a/Engine/Plugins/Experimental/SequenceNavigator/Source/SequenceNavigator/Public/Extensions/IRevisionControlExtension.h b/Engine/Plugins/Experimental/SequenceNavigator/Source/SequenceNavigator/Public/Extensions/IRevisionControlExtension.h
--- a/Engine/Plugins/Experimental/SequenceNavigator/Source/SequenceNavigator/Public/Extensions/IRevisionControlExtension.h
+++ b/Engine/Plugins/Experimental/SequenceNavigator/Source/SequenceNavigator/Public/Extensions/IRevisionControlExtension.h
@@ -9,7 +9,7 @@
 
 struct FSlateBrush;
 
-namespace UE::Sequencer::NavigationTool
+namespace UE::SequenceNavigator
 {
 
 enum class EItemRevisionControlState
@@ -20,7 +20,7 @@ enum class EItemRevisionControlState
 };
 ENUM_CLASS_FLAGS(EItemRevisionControlState)
 
-class SEQUENCENAVIGATOR_API IRevisionControlExtension : public INavigationToolItemTypeCastable
+class IRevisionControlExtension : public INavigationToolItemTypeCastable
 {
 public:
 	UE_NAVIGATIONTOOL_INHERITS(IRevisionControlExtension, INavigationToolItemTypeCastable)
@@ -30,21 +30,21 @@ public:
 		QueueRevisionControlStatusUpdate(InObject);
 	}
 
-	virtual EItemRevisionControlState GetRevisionControlState() const = 0;
+	SEQUENCENAVIGATOR_API virtual EItemRevisionControlState GetRevisionControlState() const = 0;
 
-	virtual const FSlateBrush* GetRevisionControlStatusIcon() const = 0;
+	SEQUENCENAVIGATOR_API virtual const FSlateBrush* GetRevisionControlStatusIcon() const = 0;
 
-	virtual FText GetRevisionControlStatusText() const = 0;
+	SEQUENCENAVIGATOR_API virtual FText GetRevisionControlStatusText() const = 0;
 
 	static void QueueRevisionControlStatusUpdate(const UObject* const InObject)
 	{
-		if (!IsValid(InObject))
+		if (!InObject)
 		{
 			return;
 		}
 
 		UPackage* const InPackage = InObject->GetPackage();
-		if (!IsValid(InPackage))
+		if (!InPackage)
 		{
 			return;
 		}
@@ -55,4 +55,4 @@ public:
 	}
 };
 
-}
+} // namespace UE::SequenceNavigator
diff --git a/Engine/Plugins/Experimental/SequenceNavigator/Source/SequenceNavigator/Public/Extensions/ISequenceInactivableExtension.h b/Engine/Plugins/Experimental/SequenceNavigator/Source/SequenceNavigator/Public/Extensions/ISequenceInactivableExtension.h
--- a/Engine/Plugins/Experimental/SequenceNavigator/Source/SequenceNavigator/Public/Extensions/ISequenceInactivableExtension.h
+++ b/Engine/Plugins/Experimental/SequenceNavigator/Source/SequenceNavigator/Public/Extensions/ISequenceInactivableExtension.h
@@ -4,7 +4,7 @@
 
 #include "NavigationToolItemType.h"
 
-namespace UE::Sequencer::NavigationTool
+namespace UE::SequenceNavigator
 {
 
 enum class EItemSequenceInactiveState
@@ -15,14 +15,14 @@ enum class EItemSequenceInactiveState
 };
 ENUM_CLASS_FLAGS(EItemSequenceInactiveState)
 
-class SEQUENCENAVIGATOR_API ISequenceInactivableExtension : public INavigationToolItemTypeCastable
+class ISequenceInactivableExtension : public INavigationToolItemTypeCastable
 {
 public:
 	UE_NAVIGATIONTOOL_INHERITS(ISequenceInactivableExtension, INavigationToolItemTypeCastable)
 
-	virtual EItemSequenceInactiveState GetInactiveState() const = 0;
+	SEQUENCENAVIGATOR_API virtual EItemSequenceInactiveState GetInactiveState() const = 0;
 
-	virtual void SetIsInactive(const bool bInIsMuted) = 0;
+	SEQUENCENAVIGATOR_API virtual void SetIsInactive(const bool bInIsMuted) = 0;
 };
 
-}
+} // namespace UE::SequenceNavigator
diff --git a/Engine/Plugins/Experimental/SequenceNavigator/Source/SequenceNavigator/Public/Extensions/ISequenceLockableExtension.h b/Engine/Plugins/Experimental/SequenceNavigator/Source/SequenceNavigator/Public/Extensions/ISequenceLockableExtension.h
--- a/Engine/Plugins/Experimental/SequenceNavigator/Source/SequenceNavigator/Public/Extensions/ISequenceLockableExtension.h
+++ b/Engine/Plugins/Experimental/SequenceNavigator/Source/SequenceNavigator/Public/Extensions/ISequenceLockableExtension.h
@@ -4,7 +4,7 @@
 
 #include "NavigationToolItemType.h"
 
-namespace UE::Sequencer::NavigationTool
+namespace UE::SequenceNavigator
 {
 
 enum class EItemSequenceLockState
@@ -15,14 +15,14 @@ enum class EItemSequenceLockState
 };
 ENUM_CLASS_FLAGS(EItemSequenceLockState)
 
-class SEQUENCENAVIGATOR_API ISequenceLockableExtension : public INavigationToolItemTypeCastable
+class ISequenceLockableExtension : public INavigationToolItemTypeCastable
 {
 public:
 	UE_NAVIGATIONTOOL_INHERITS(ISequenceLockableExtension, INavigationToolItemTypeCastable)
 
-	virtual EItemSequenceLockState GetLockState() const = 0;
+	SEQUENCENAVIGATOR_API virtual EItemSequenceLockState GetLockState() const = 0;
 
-	virtual void SetIsLocked(const bool bInIsLocked) = 0;
+	SEQUENCENAVIGATOR_API virtual void SetIsLocked(const bool bInIsLocked) = 0;
 };
 
-}
+} // namespace UE::SequenceNavigator
diff --git a/Engine/Plugins/Experimental/SequenceNavigator/Source/SequenceNavigator/Public/Extensions/ISequenceMutableExtension.h b/Engine/Plugins/Experimental/SequenceNavigator/Source/SequenceNavigator/Public/Extensions/ISequenceMutableExtension.h
--- a/Engine/Plugins/Experimental/SequenceNavigator/Source/SequenceNavigator/Public/Extensions/ISequenceMutableExtension.h
+++ b/Engine/Plugins/Experimental/SequenceNavigator/Source/SequenceNavigator/Public/Extensions/ISequenceMutableExtension.h
@@ -4,7 +4,7 @@
 
 #include "NavigationToolItemType.h"
 
-namespace UE::Sequencer::NavigationTool
+namespace UE::SequenceNavigator
 {
 
 enum class EItemSequenceMuteState
@@ -15,14 +15,14 @@ enum class EItemSequenceMuteState
 };
 ENUM_CLASS_FLAGS(EItemSequenceMuteState)
 
-class SEQUENCENAVIGATOR_API ISequenceMutableExtension : public INavigationToolItemTypeCastable
+class ISequenceMutableExtension : public INavigationToolItemTypeCastable
 {
 public:
 	UE_NAVIGATIONTOOL_INHERITS(ISequenceMutableExtension, INavigationToolItemTypeCastable)
 
-	virtual EItemSequenceMuteState GetMuteState() const = 0;
+	SEQUENCENAVIGATOR_API virtual EItemSequenceMuteState GetMuteState() const = 0;
 
-	virtual void SetIsMuted(const bool bInIsMuted) = 0;
+	SEQUENCENAVIGATOR_API virtual void SetIsMuted(const bool bInIsMuted) = 0;
 };
 
-}
+} // namespace UE::SequenceNavigator
diff --git a/Engine/Plugins/Experimental/SequenceNavigator/Source/SequenceNavigator/Public/Filters/Filters/NavigationToolBuiltInFilterParams.h b/Engine/Plugins/Experimental/SequenceNavigator/Source/SequenceNavigator/Public/Filters/Filters/NavigationToolBuiltInFilterParams.h
--- a/Engine/Plugins/Experimental/SequenceNavigator/Source/SequenceNavigator/Public/Filters/Filters/NavigationToolBuiltInFilterParams.h
+++ b/Engine/Plugins/Experimental/SequenceNavigator/Source/SequenceNavigator/Public/Filters/Filters/NavigationToolBuiltInFilterParams.h
@@ -9,7 +9,7 @@
 
 class UObject;
 
-namespace UE::Sequencer::NavigationTool
+namespace UE::SequenceNavigator
 {
 	class FNavigationToolItemTypeId;
 }
@@ -27,7 +27,7 @@ struct FNavigationToolBuiltInFilterParams
 	FNavigationToolBuiltInFilterParams() = default;
 
 	SEQUENCENAVIGATOR_API FNavigationToolBuiltInFilterParams(const FName InFilterId
-		, const TSet<UE::Sequencer::NavigationTool::FNavigationToolItemTypeId>& InItemClasses
+		, const TSet<UE::SequenceNavigator::FNavigationToolItemTypeId>& InItemClasses
 		, const TArray<TSubclassOf<UObject>>& InFilterClasses
 		, const ENavigationToolFilterMode InMode = ENavigationToolFilterMode::MatchesType
 		, const FSlateBrush* const InIconBrush = nullptr
@@ -56,18 +56,18 @@ struct FNavigationToolBuiltInFilterParams
 
 	void SetOverrideIconColor(const FSlateColor& InNewIconColor);
 
-	bool IsValidItemClass(const UE::Sequencer::NavigationTool::FNavigationToolItemTypeId InClassTypeId) const;
+	bool IsValidItemClass(const UE::SequenceNavigator::FNavigationToolItemTypeId InClassTypeId) const;
 	bool IsValidObjectClass(const UClass* const InClass) const;
 
 	void SetFilterText(const FText& InText);
 
-	bool PassesFilterText(const UE::Sequencer::NavigationTool::FNavigationToolItemPtr& InItem) const;
+	bool PassesFilterText(const UE::SequenceNavigator::FNavigationToolItemPtr& InItem) const;
 
 private:
 	UPROPERTY(EditAnywhere, Category = "Filter")
 	FName FilterId = NAME_None;
 
-	TSet<UE::Sequencer::NavigationTool::FNavigationToolItemTypeId> ItemClasses;
+	TSet<UE::SequenceNavigator::FNavigationToolItemTypeId> ItemClasses;
 
 	UPROPERTY(EditAnywhere, Category = "Filter")
 	TArray<TSubclassOf<UObject>> ObjectClasses;
diff --git a/Engine/Plugins/Experimental/SequenceNavigator/Source/SequenceNavigator/Public/Filters/Filters/NavigationToolFilterExtension.h b/Engine/Plugins/Experimental/SequenceNavigator/Source/SequenceNavigator/Public/Filters/Filters/NavigationToolFilterExtension.h
--- a/Engine/Plugins/Experimental/SequenceNavigator/Source/SequenceNavigator/Public/Filters/Filters/NavigationToolFilterExtension.h
+++ b/Engine/Plugins/Experimental/SequenceNavigator/Source/SequenceNavigator/Public/Filters/Filters/NavigationToolFilterExtension.h
@@ -6,7 +6,7 @@
 
 class FFilterCategory;
 
-namespace UE::Sequencer::NavigationTool
+namespace UE::SequenceNavigator
 {
 	class FNavigationToolFilter;
 	class INavigationToolFilterBar;
@@ -27,8 +27,8 @@ public:
 	 * @param InPreferredCategory  The preferred filter bar category to place any newly added filters to, but can be any category
 	 * @param InOutFilterList      Filter list to add additional filters to
 	 */
-	virtual void AddFilterExtensions(UE::Sequencer::NavigationTool::INavigationToolFilterBar& InOutFilterInterface
+	virtual void AddFilterExtensions(UE::SequenceNavigator::INavigationToolFilterBar& InOutFilterInterface
 		, const TSharedRef<FFilterCategory>& InPreferredCategory
-		, TArray<TSharedRef<UE::Sequencer::NavigationTool::FNavigationToolFilter>>& InOutFilterList) const
+		, TArray<TSharedRef<UE::SequenceNavigator::FNavigationToolFilter>>& InOutFilterList) const
 	{}
 };
diff --git a/Engine/Plugins/Experimental/SequenceNavigator/Source/SequenceNavigator/Public/Filters/TextExpressions/NavigationToolFilterTextExpressionExtension.h b/Engine/Plugins/Experimental/SequenceNavigator/Source/SequenceNavigator/Public/Filters/TextExpressions/NavigationToolFilterTextExpressionExtension.h
--- a/Engine/Plugins/Experimental/SequenceNavigator/Source/SequenceNavigator/Public/Filters/TextExpressions/NavigationToolFilterTextExpressionExtension.h
+++ b/Engine/Plugins/Experimental/SequenceNavigator/Source/SequenceNavigator/Public/Filters/TextExpressions/NavigationToolFilterTextExpressionExtension.h
@@ -6,7 +6,7 @@
 
 class FNavigationToolFilterTextExpressionContext;
 
-namespace UE::Sequencer::NavigationTool
+namespace UE::SequenceNavigator
 {
 
 class INavigationToolFilterBar;
@@ -28,7 +28,7 @@ public:
 	 * @param InOutFilterInterface The filter interface to extend
 	 * @param InOutExpressionList  Expression list to add additional text expressions to
 	 */
-	virtual void AddFilterTextExpressionExtensions(UE::Sequencer::NavigationTool::INavigationToolFilterBar& InOutFilterInterface
-		, TArray<TSharedRef<UE::Sequencer::NavigationTool::FNavigationToolFilterTextExpressionContext>>& InOutExpressionList) const
+	virtual void AddFilterTextExpressionExtensions(UE::SequenceNavigator::INavigationToolFilterBar& InOutFilterInterface
+		, TArray<TSharedRef<UE::SequenceNavigator::FNavigationToolFilterTextExpressionContext>>& InOutExpressionList) const
 	{}
 };
diff --git a/Engine/Plugins/Experimental/SequenceNavigator/Source/SequenceNavigator/Public/INavigationTool.h b/Engine/Plugins/Experimental/SequenceNavigator/Source/SequenceNavigator/Public/INavigationTool.h
--- a/Engine/Plugins/Experimental/SequenceNavigator/Source/SequenceNavigator/Public/INavigationTool.h
+++ b/Engine/Plugins/Experimental/SequenceNavigator/Source/SequenceNavigator/Public/INavigationTool.h
@@ -14,13 +14,14 @@
 #include "Modules/ModuleManager.h"
 #include "NavigationToolDefines.h"
 #include "NavigationToolExtender.h"
+#include "SequencerCoreFwd.h"
 #include "Templates/SharedPointer.h"
 
 class FUICommandList;
 class ISequencer;
 template<typename OptionalType> struct TOptional;
 
-namespace UE::Sequencer::NavigationTool
+namespace UE::SequenceNavigator
 {
 
 class FNavigationToolProvider;
@@ -175,7 +176,7 @@ public:
 	 * @param InViewModel The Sequencer view model to find the items of
 	 * @return The list of Navigation Tool items found
 	 */
-	virtual TArray<FNavigationToolItemPtr> TryFindItems(const FViewModelPtr& InViewModel) const = 0;
+	virtual TArray<FNavigationToolItemPtr> TryFindItems(const UE::Sequencer::FViewModelPtr& InViewModel) const = 0;
 
 	/**
  	 * Adds or Removes the Ignore Notify Flags to prevent certain actions from automatically happening when they're triggered
@@ -252,4 +253,4 @@ protected:
 	virtual const FNavigationToolItemProxyRegistry& GetItemProxyRegistry() const = 0;
 };
 
-}
+} // namespace UE::SequenceNavigator
diff --git a/Engine/Plugins/Experimental/SequenceNavigator/Source/SequenceNavigator/Public/Items/NavigationToolBinding.h b/Engine/Plugins/Experimental/SequenceNavigator/Source/SequenceNavigator/Public/Items/NavigationToolBinding.h
--- a/Engine/Plugins/Experimental/SequenceNavigator/Source/SequenceNavigator/Public/Items/NavigationToolBinding.h
+++ b/Engine/Plugins/Experimental/SequenceNavigator/Source/SequenceNavigator/Public/Items/NavigationToolBinding.h
@@ -8,9 +8,15 @@
 #include "MovieSceneBinding.h"
 #include "MVVM/ViewModels/ObjectBindingModel.h"
 #include "NavigationToolItem.h"
+#include "SequencerCoreFwd.h"
 #include "Textures/SlateIcon.h"
 
-namespace UE::Sequencer::NavigationTool
+namespace UE::Sequencer
+{
+class FObjectBindingModel;
+}
+
+namespace UE::SequenceNavigator
 {
 
 class FNavigationToolSequence;
@@ -92,7 +98,7 @@ public:
 	/** Caches the object if it's not already been cached and returns the result. */
 	UObject* CacheBoundObject();
 
-	TViewModelPtr<FObjectBindingModel> GetViewModel() const;
+	UE::Sequencer::TViewModelPtr<UE::Sequencer::FObjectBindingModel> GetViewModel() const;
 
 protected:
 	TWeakPtr<FNavigationToolSequence> WeakParentSequenceItem;
@@ -104,4 +110,4 @@ protected:
 	FSlateIcon Icon;
 };
 
-}
+} // namespace UE::SequenceNavigator
diff --git a/Engine/Plugins/Experimental/SequenceNavigator/Source/SequenceNavigator/Public/Items/NavigationToolItemId.h b/Engine/Plugins/Experimental/SequenceNavigator/Source/SequenceNavigator/Public/Items/NavigationToolItemId.h
--- a/Engine/Plugins/Experimental/SequenceNavigator/Source/SequenceNavigator/Public/Items/NavigationToolItemId.h
+++ b/Engine/Plugins/Experimental/SequenceNavigator/Source/SequenceNavigator/Public/Items/NavigationToolItemId.h
@@ -10,7 +10,7 @@ class UMovieSceneSequence;
 class UMovieSceneSection;
 class UObject;
 
-namespace UE::Sequencer::NavigationTool
+namespace UE::SequenceNavigator
 {
 
 class INavigationToolItem;
@@ -86,7 +86,7 @@ struct SEQUENCENAVIGATOR_API FNavigationToolItemId
 	 * Ctor used for making the Item Id from a Sequencer view model
 	 * @param InViewModel The view model to construct the Id from
 	 */
-	FNavigationToolItemId(const FViewModelPtr& InViewModel);
+	FNavigationToolItemId(const UE::Sequencer::FViewModelPtr& InViewModel);
 
 	FNavigationToolItemId(const FNavigationToolItemId& Other);
 	FNavigationToolItemId(FNavigationToolItemId&& Other) noexcept;
@@ -129,4 +129,4 @@ private:
 	bool bHasCachedHash = false;
 };
 
-}
+} // namespace UE::SequenceNavigator
diff --git a/Engine/Plugins/Experimental/SequenceNavigator/Source/SequenceNavigator/Public/Items/NavigationToolItemParameters.h b/Engine/Plugins/Experimental/SequenceNavigator/Source/SequenceNavigator/Public/Items/NavigationToolItemParameters.h
--- a/Engine/Plugins/Experimental/SequenceNavigator/Source/SequenceNavigator/Public/Items/NavigationToolItemParameters.h
+++ b/Engine/Plugins/Experimental/SequenceNavigator/Source/SequenceNavigator/Public/Items/NavigationToolItemParameters.h
@@ -7,7 +7,7 @@
 #include "NavigationToolDefines.h"
 #include "Widgets/Views/STableRow.h"
 
-namespace UE::Sequencer::NavigationTool
+namespace UE::SequenceNavigator
 {
 
 class FNavigationToolItem;
diff --git a/Engine/Plugins/Experimental/SequenceNavigator/Source/SequenceNavigator/Public/Items/NavigationToolSequence.h b/Engine/Plugins/Experimental/SequenceNavigator/Source/SequenceNavigator/Public/Items/NavigationToolSequence.h
--- a/Engine/Plugins/Experimental/SequenceNavigator/Source/SequenceNavigator/Public/Items/NavigationToolSequence.h
+++ b/Engine/Plugins/Experimental/SequenceNavigator/Source/SequenceNavigator/Public/Items/NavigationToolSequence.h
@@ -11,6 +11,7 @@
 #include "Extensions/ISequenceInactivableExtension.h"
 #include "Extensions/ISequenceLockableExtension.h"
 #include "NavigationToolItem.h"
+#include "SequencerCoreFwd.h"
 
 class UMovieScene;
 class UMovieSceneSequence;
@@ -22,7 +23,7 @@ namespace UE::Sequencer
 class FSectionModel;
 }
 
-namespace UE::Sequencer::NavigationTool
+namespace UE::SequenceNavigator
 {
 
 /**
@@ -141,7 +142,7 @@ public:
 
 	SEQUENCENAVIGATOR_API TArray<FMovieSceneBinding> GetSortedBindings() const;
 
-	TViewModelPtr<FSectionModel> GetViewModel() const;
+	UE::Sequencer::TViewModelPtr<UE::Sequencer::FSectionModel> GetViewModel() const;
 
 protected:
 	//~Begin FNavigationToolItem
@@ -153,4 +154,4 @@ protected:
 	TWeakObjectPtr<UMovieSceneSequence> WeakSequence;
 };
 
-}
+} // namespace UE::SequenceNavigator
diff --git a/Engine/Plugins/Experimental/SequenceNavigator/Source/SequenceNavigator/Public/Items/NavigationToolTrack.h b/Engine/Plugins/Experimental/SequenceNavigator/Source/SequenceNavigator/Public/Items/NavigationToolTrack.h
--- a/Engine/Plugins/Experimental/SequenceNavigator/Source/SequenceNavigator/Public/Items/NavigationToolTrack.h
+++ b/Engine/Plugins/Experimental/SequenceNavigator/Source/SequenceNavigator/Public/Items/NavigationToolTrack.h
@@ -5,7 +5,9 @@
 #include "Extensions/IPlayheadExtension.h"
 #include "Extensions/ISequenceLockableExtension.h"
 #include "NavigationToolItem.h"
+#include "SequencerCoreFwd.h"
 #include "Textures/SlateIcon.h"
+#include "UObject/WeakObjectPtrTemplates.h"
 
 class UMovieSceneSequence;
 class UMovieSceneSection;
@@ -16,7 +18,7 @@ namespace UE::Sequencer
 class FTrackModel;
 }
 
-namespace UE::Sequencer::NavigationTool
+namespace UE::SequenceNavigator
 {
 
 class FNavigationToolTrack
@@ -78,7 +80,7 @@ public:
 	UMovieSceneSection* GetSection() const;
 	int32 GetSectionIndex() const;
 
-	TViewModelPtr<FTrackModel> GetViewModel() const;
+	UE::Sequencer::TViewModelPtr<UE::Sequencer::FTrackModel> GetViewModel() const;
 
 protected:
 	//~ Begin FNavigationToolItem
@@ -95,4 +97,4 @@ protected:
 	FSlateIcon Icon;
 };
 
-}
+} // namespace UE::SequenceNavigator
diff --git a/Engine/Plugins/Experimental/SequenceNavigator/Source/SequenceNavigator/Public/Menus/NavigationToolFilterBarContext.h b/Engine/Plugins/Experimental/SequenceNavigator/Source/SequenceNavigator/Public/Menus/NavigationToolFilterBarContext.h
--- a/Engine/Plugins/Experimental/SequenceNavigator/Source/SequenceNavigator/Public/Menus/NavigationToolFilterBarContext.h
+++ b/Engine/Plugins/Experimental/SequenceNavigator/Source/SequenceNavigator/Public/Menus/NavigationToolFilterBarContext.h
@@ -16,12 +16,12 @@ class UNavigationToolFilterBarContext : public UObject
 	GENERATED_BODY()
 
 public:
-	void Init(const TWeakPtr<UE::Sequencer::NavigationTool::INavigationToolFilterBar>& InFilterBarWeak)
+	void Init(const TWeakPtr<UE::SequenceNavigator::INavigationToolFilterBar>& InFilterBarWeak)
 	{
 		WeakFilterBar = InFilterBarWeak;
 	}
 
-	TSharedPtr<UE::Sequencer::NavigationTool::INavigationToolFilterBar> GetFilterBar() const
+	TSharedPtr<UE::SequenceNavigator::INavigationToolFilterBar> GetFilterBar() const
 	{
 		return WeakFilterBar.IsValid() ? WeakFilterBar.Pin() : nullptr;
 	}
@@ -29,5 +29,5 @@ public:
 	FOnPopulateFilterBarMenu OnPopulateMenu;
 
 protected:
-	TWeakPtr<UE::Sequencer::NavigationTool::INavigationToolFilterBar> WeakFilterBar;
+	TWeakPtr<UE::SequenceNavigator::INavigationToolFilterBar> WeakFilterBar;
 };
diff --git a/Engine/Plugins/Experimental/SequenceNavigator/Source/SequenceNavigator/Public/Menus/NavigationToolFilterMenuContext.h b/Engine/Plugins/Experimental/SequenceNavigator/Source/SequenceNavigator/Public/Menus/NavigationToolFilterMenuContext.h
--- a/Engine/Plugins/Experimental/SequenceNavigator/Source/SequenceNavigator/Public/Menus/NavigationToolFilterMenuContext.h
+++ b/Engine/Plugins/Experimental/SequenceNavigator/Source/SequenceNavigator/Public/Menus/NavigationToolFilterMenuContext.h
@@ -13,12 +13,12 @@ class UNavigationToolFilterMenuContext : public UObject
 	GENERATED_BODY()
 
 public:
-	void Init(const TWeakPtr<UE::Sequencer::NavigationTool::FNavigationToolFilter>& InFilter)
+	void Init(const TWeakPtr<UE::SequenceNavigator::FNavigationToolFilter>& InFilter)
 	{
 		WeakFilter = InFilter;
 	}
 
-	TSharedPtr<UE::Sequencer::NavigationTool::FNavigationToolFilter> GetFilter() const
+	TSharedPtr<UE::SequenceNavigator::FNavigationToolFilter> GetFilter() const
 	{
 		return WeakFilter.Pin();
 	}
@@ -26,5 +26,5 @@ public:
 	FOnPopulateFilterBarMenu OnPopulateMenu;
 
 protected:
-	TWeakPtr<UE::Sequencer::NavigationTool::FNavigationToolFilter> WeakFilter;
+	TWeakPtr<UE::SequenceNavigator::FNavigationToolFilter> WeakFilter;
 };
diff --git a/Engine/Plugins/Experimental/SequenceNavigator/Source/SequenceNavigator/Public/Menus/NavigationToolItemMenuContext.h b/Engine/Plugins/Experimental/SequenceNavigator/Source/SequenceNavigator/Public/Menus/NavigationToolItemMenuContext.h
--- a/Engine/Plugins/Experimental/SequenceNavigator/Source/SequenceNavigator/Public/Menus/NavigationToolItemMenuContext.h
+++ b/Engine/Plugins/Experimental/SequenceNavigator/Source/SequenceNavigator/Public/Menus/NavigationToolItemMenuContext.h
@@ -9,7 +9,7 @@
 #include "UObject/Object.h"
 #include "NavigationToolItemMenuContext.generated.h"
 
-namespace UE::Sequencer::NavigationTool
+namespace UE::SequenceNavigator
 {
 class INavigationTool;
 }
@@ -22,24 +22,24 @@ class UNavigationToolItemMenuContext : public UObject
 	friend class FNavigationToolView;
 
 public:
-	void Init(const TWeakPtr<UE::Sequencer::NavigationTool::INavigationTool>& InTool, const TArray<UE::Sequencer::NavigationTool::FNavigationToolItemPtr>& InItemList)
+	void Init(const TWeakPtr<UE::SequenceNavigator::INavigationTool>& InTool, const TArray<UE::SequenceNavigator::FNavigationToolItemPtr>& InItemList)
 	{
 		WeakTool = InTool;
 		WeakItemList.Append(InItemList);
 	}
 
-	TSharedPtr<UE::Sequencer::NavigationTool::INavigationTool> GetTool() const
+	TSharedPtr<UE::SequenceNavigator::INavigationTool> GetTool() const
 	{
 		return WeakTool.Pin();
 	}
 
-	TConstArrayView<UE::Sequencer::NavigationTool::FNavigationToolItemWeakPtr> GetItems() const
+	TConstArrayView<UE::SequenceNavigator::FNavigationToolItemWeakPtr> GetItems() const
 	{
 		return WeakItemList;
 	}
 
 private:
-	TWeakPtr<UE::Sequencer::NavigationTool::INavigationTool> WeakTool;
+	TWeakPtr<UE::SequenceNavigator::INavigationTool> WeakTool;
 
-	TArray<UE::Sequencer::NavigationTool::FNavigationToolItemWeakPtr> WeakItemList;
+	TArray<UE::SequenceNavigator::FNavigationToolItemWeakPtr> WeakItemList;
 };
diff --git a/Engine/Plugins/Experimental/SequenceNavigator/Source/SequenceNavigator/Public/NavigationToolDefines.h b/Engine/Plugins/Experimental/SequenceNavigator/Source/SequenceNavigator/Public/NavigationToolDefines.h
--- a/Engine/Plugins/Experimental/SequenceNavigator/Source/SequenceNavigator/Public/NavigationToolDefines.h
+++ b/Engine/Plugins/Experimental/SequenceNavigator/Source/SequenceNavigator/Public/NavigationToolDefines.h
@@ -6,10 +6,7 @@
 #include "Templates/SharedPointer.h"
 #include "UObject/ObjectMacros.h"
 
-class AActor;
-class USceneComponent;
-
-namespace UE::Sequencer::NavigationTool
+namespace UE::SequenceNavigator
 {
 
 class INavigationToolItem;
@@ -21,7 +18,7 @@ using FNavigationToolItemWeakPtr = TWeakPtr<INavigationToolItem>;
 SEQUENCENAVIGATOR_API FName GetToolBarMenuName();
 SEQUENCENAVIGATOR_API FName GetItemContextMenuName();
 
-}
+} // namespace UE::SequenceNavigator
 
 /** Flags specifying how an Item should be Selected */
 enum class ENavigationToolItemSelectionFlags : uint8
diff --git a/Engine/Plugins/MovieScene/LevelSequenceEditor/Source/LevelSequenceEditor/Private/LevelSequenceEditorToolkit.cpp b/Engine/Plugins/MovieScene/LevelSequenceEditor/Source/LevelSequenceEditor/Private/LevelSequenceEditorToolkit.cpp
--- a/Engine/Plugins/MovieScene/LevelSequenceEditor/Source/LevelSequenceEditor/Private/LevelSequenceEditorToolkit.cpp
+++ b/Engine/Plugins/MovieScene/LevelSequenceEditor/Source/LevelSequenceEditor/Private/LevelSequenceEditorToolkit.cpp
@@ -320,7 +320,7 @@ void FLevelSequenceEditorToolkit::Initialize(const EToolkitMode::Type Mode, cons
 
 	if (Sequencer.IsValid())
 	{
-		using namespace UE::Sequencer::NavigationTool;
+		using namespace UE::SequenceNavigator;
 
 		const FName ToolId = FNavigationToolExtender::GetToolInstanceId(*Sequencer);
 
@@ -999,7 +999,7 @@ void FLevelSequenceEditorToolkit::OnClose()
 {
 	if (Sequencer.IsValid() && NavigationToolProvider.IsValid())
 	{
-		using namespace UE::Sequencer::NavigationTool;
+		using namespace UE::SequenceNavigator;
 
 		const FName ToolId = FNavigationToolExtender::GetToolInstanceId(*Sequencer);
 		FNavigationToolExtender::UnregisterToolProvider(ToolId, NavigationToolProvider->GetIdentifier());
diff --git a/Engine/Plugins/MovieScene/LevelSequenceEditor/Source/LevelSequenceEditor/Private/LevelSequenceNavigationToolProvider.cpp b/Engine/Plugins/MovieScene/LevelSequenceEditor/Source/LevelSequenceEditor/Private/LevelSequenceNavigationToolProvider.cpp
--- a/Engine/Plugins/MovieScene/LevelSequenceEditor/Source/LevelSequenceEditor/Private/LevelSequenceNavigationToolProvider.cpp
+++ b/Engine/Plugins/MovieScene/LevelSequenceEditor/Source/LevelSequenceEditor/Private/LevelSequenceNavigationToolProvider.cpp
@@ -78,9 +78,9 @@ void FLevelSequenceNavigationToolProvider::OnDeactivate()
 		.Remove(DragDropInitializedDelegate);*/
 }
 
-void FLevelSequenceNavigationToolProvider::OnExtendColumns(UE::Sequencer::NavigationTool::FNavigationToolColumnExtender& OutExtender)
+void FLevelSequenceNavigationToolProvider::OnExtendColumns(UE::SequenceNavigator::FNavigationToolColumnExtender& OutExtender)
 {
-	using namespace UE::Sequencer::NavigationTool;
+	using namespace UE::SequenceNavigator;
 
 	OutExtender.AddColumn<FNavigationToolPlayheadColumn>();
 	OutExtender.AddColumn<FNavigationToolDeactiveStateColumn>();
@@ -108,7 +108,7 @@ void FLevelSequenceNavigationToolProvider::OnExtendColumns(UE::Sequencer::Naviga
 
 void FLevelSequenceNavigationToolProvider::OnExtendColumnViews(TSet<FNavigationToolColumnView>& OutColumnViews)
 {
-	using namespace UE::Sequencer::NavigationTool;
+	using namespace UE::SequenceNavigator;
 
 	FNavigationToolColumnView& AnimationColumnView = OutColumnViews
 		.FindOrAdd(FNavigationToolColumnView(AnimationColumnViewName));
@@ -127,12 +127,12 @@ void FLevelSequenceNavigationToolProvider::OnExtendColumnViews(TSet<FNavigationT
 	FNavigationToolProvider::OnExtendColumnViews(OutColumnViews);
 }
 
-void FLevelSequenceNavigationToolProvider::OnExtendItemChildren(UE::Sequencer::NavigationTool::INavigationTool& InTool
-	, const UE::Sequencer::NavigationTool::FNavigationToolItemPtr& InParentItem
-	, TArray<UE::Sequencer::NavigationTool::FNavigationToolItemPtr>& OutChildren
+void FLevelSequenceNavigationToolProvider::OnExtendItemChildren(UE::SequenceNavigator::INavigationTool& InTool
+	, const UE::SequenceNavigator::FNavigationToolItemPtr& InParentItem
+	, TArray<UE::SequenceNavigator::FNavigationToolItemPtr>& OutChildren
 	, const bool bInRecursive)
 {
-	using namespace UE::Sequencer::NavigationTool;
+	using namespace UE::SequenceNavigator;
 
 	FNavigationToolProvider::OnExtendItemChildren(InTool, InParentItem, OutChildren, bInRecursive);
 
@@ -149,7 +149,7 @@ void FLevelSequenceNavigationToolProvider::OnExtendItemChildren(UE::Sequencer::N
 	}
 
 	UMovieSceneSequence* const RootSequence = Sequencer->GetRootMovieSceneSequence();
-	if (!IsValid(RootSequence))
+	if (!RootSequence)
 	{
 		return;
 	}
@@ -175,7 +175,7 @@ void FLevelSequenceNavigationToolProvider::OnExtendBuiltInFilters(TArray<FNaviga
 	FNavigationToolProvider::OnExtendBuiltInFilters(OutFilterParams);
 }
 
-FMovieSceneEditorData* FLevelSequenceNavigationToolProvider::GetRootMovieSceneEditorData(const UE::Sequencer::NavigationTool::INavigationTool& InTool) const
+FMovieSceneEditorData* FLevelSequenceNavigationToolProvider::GetRootMovieSceneEditorData(const UE::SequenceNavigator::INavigationTool& InTool) const
 {
 	const TSharedPtr<ISequencer> Sequencer = InTool.GetSequencer();
 	if (!Sequencer.IsValid())
@@ -184,13 +184,13 @@ FMovieSceneEditorData* FLevelSequenceNavigationToolProvider::GetRootMovieSceneEd
 	}
 
 	UMovieSceneSequence* const RootMovieSceneSequence = Sequencer->GetRootMovieSceneSequence();
-	if (!IsValid(RootMovieSceneSequence))
+	if (!RootMovieSceneSequence)
 	{
 		return nullptr;
 	}
 
 	UMovieScene* const RootMovieScene = RootMovieSceneSequence->GetMovieScene();
-	if (!IsValid(RootMovieScene))
+	if (!RootMovieScene)
 	{
 		return nullptr;
 	}
@@ -198,7 +198,7 @@ FMovieSceneEditorData* FLevelSequenceNavigationToolProvider::GetRootMovieSceneEd
 	return &RootMovieScene->GetEditorData();
 }
 
-FNavigationToolSaveState* FLevelSequenceNavigationToolProvider::GetSaveState(const UE::Sequencer::NavigationTool::INavigationTool& InTool) const
+FNavigationToolSaveState* FLevelSequenceNavigationToolProvider::GetSaveState(const UE::SequenceNavigator::INavigationTool& InTool) const
 {
 	if (FMovieSceneEditorData* const EditorData = GetRootMovieSceneEditorData(InTool))
 	{
@@ -207,7 +207,7 @@ FNavigationToolSaveState* FLevelSequenceNavigationToolProvider::GetSaveState(con
 	return nullptr;
 }
 
-void FLevelSequenceNavigationToolProvider::SetSaveState(const UE::Sequencer::NavigationTool::INavigationTool& InTool
+void FLevelSequenceNavigationToolProvider::SetSaveState(const UE::SequenceNavigator::INavigationTool& InTool
 	, const FNavigationToolSaveState& InSaveState) const
 {
 	if (FMovieSceneEditorData* const EditorData = GetRootMovieSceneEditorData(InTool))
diff --git a/Engine/Plugins/MovieScene/LevelSequenceEditor/Source/LevelSequenceEditor/Private/LevelSequenceNavigationToolProvider.h b/Engine/Plugins/MovieScene/LevelSequenceEditor/Source/LevelSequenceEditor/Private/LevelSequenceNavigationToolProvider.h
--- a/Engine/Plugins/MovieScene/LevelSequenceEditor/Source/LevelSequenceEditor/Private/LevelSequenceNavigationToolProvider.h
+++ b/Engine/Plugins/MovieScene/LevelSequenceEditor/Source/LevelSequenceEditor/Private/LevelSequenceNavigationToolProvider.h
@@ -13,14 +13,14 @@ struct FMovieSceneEditorData;
 struct FNavigationToolBuiltInFilterParams;
 struct FNavigationToolSaveState;
 
-namespace UE::Sequencer::NavigationTool
+namespace UE::SequenceNavigator
 {
 	class FNavigationToolColumnExtender;
 	class FNavigationToolItem;
 	class INavigationTool;
 }
 
-class FLevelSequenceNavigationToolProvider : public UE::Sequencer::NavigationTool::FNavigationToolProvider
+class FLevelSequenceNavigationToolProvider : public UE::SequenceNavigator::FNavigationToolProvider
 {
 public:
 	static const FName Identifier;
@@ -43,22 +43,22 @@ public:
 	virtual void OnActivate() override;
 	virtual void OnDeactivate() override;
 
-	virtual void OnExtendColumns(UE::Sequencer::NavigationTool::FNavigationToolColumnExtender& OutExtender) override;
+	virtual void OnExtendColumns(UE::SequenceNavigator::FNavigationToolColumnExtender& OutExtender) override;
 	virtual void OnExtendColumnViews(TSet<FNavigationToolColumnView>& OutColumnViews) override;
-	virtual void OnExtendItemChildren(UE::Sequencer::NavigationTool::INavigationTool& InTool
-		, const UE::Sequencer::NavigationTool::FNavigationToolItemPtr& InParentItem
-		, TArray<UE::Sequencer::NavigationTool::FNavigationToolItemPtr>& OutChildren
+	virtual void OnExtendItemChildren(UE::SequenceNavigator::INavigationTool& InTool
+		, const UE::SequenceNavigator::FNavigationToolItemPtr& InParentItem
+		, TArray<UE::SequenceNavigator::FNavigationToolItemPtr>& OutChildren
 		, const bool bInRecursive) override;
 	virtual void OnExtendBuiltInFilters(TArray<FNavigationToolBuiltInFilterParams>& OutFilterParams) override;
 
-	virtual FNavigationToolSaveState* GetSaveState(const UE::Sequencer::NavigationTool::INavigationTool& InTool) const override;
-	virtual void SetSaveState(const UE::Sequencer::NavigationTool::INavigationTool& InTool
+	virtual FNavigationToolSaveState* GetSaveState(const UE::SequenceNavigator::INavigationTool& InTool) const override;
+	virtual void SetSaveState(const UE::SequenceNavigator::INavigationTool& InTool
 		, const FNavigationToolSaveState& InSaveState) const override;
 
 	//~ End INavigationToolProvider
 
 private:
-	FMovieSceneEditorData* GetRootMovieSceneEditorData(const UE::Sequencer::NavigationTool::INavigationTool& InTool) const;
+	FMovieSceneEditorData* GetRootMovieSceneEditorData(const UE::SequenceNavigator::INavigationTool& InTool) const;
 
 	TSharedRef<FUICommandList> ToolCommands;
 };
