Fix an issue where every instance of SFilmOverlayOptions would push new overlays to the static UFilmOverlayToolkit, destroying previous state.

diff --git a/Engine/Plugins/MovieScene/LevelSequenceEditor/Source/LevelSequenceEditor/Private/CinematicViewport/FilmOverlays.cpp b/Engine/Plugins/MovieScene/LevelSequenceEditor/Source/LevelSequenceEditor/Private/CinematicViewport/FilmOverlays.cpp
--- a/Engine/Plugins/MovieScene/LevelSequenceEditor/Source/LevelSequenceEditor/Private/CinematicViewport/FilmOverlays.cpp
+++ b/Engine/Plugins/MovieScene/LevelSequenceEditor/Source/LevelSequenceEditor/Private/CinematicViewport/FilmOverlays.cpp
@@ -502,32 +502,30 @@ void SFilmOverlayOptions::Construct(const FArguments& InArgs, TSharedPtr<SFilmOv
 {
 	PrimaryColorTint = FLinearColor(1.f, 1.f, 1.f, .5f);
 
-	PrimaryOverlays.Add(MakeShareable(new FFilmOverlay_None));
-	UFilmOverlayToolkit::RegisterPrimaryFilmOverlay(NAME_None, PrimaryOverlays.Last());
-
-	PrimaryOverlays.Add(MakeShareable(new FFilmOverlay_Grid(3, 3)));
-	UFilmOverlayToolkit::RegisterPrimaryFilmOverlay("Grid3x3", PrimaryOverlays.Last());
-
-	PrimaryOverlays.Add(MakeShareable(new FFilmOverlay_Grid(2, 2)));
-	UFilmOverlayToolkit::RegisterPrimaryFilmOverlay("Grid2x2", PrimaryOverlays.Last());
-
-	PrimaryOverlays.Add(MakeShareable(new FFilmOverlay_Crosshair));
-	UFilmOverlayToolkit::RegisterPrimaryFilmOverlay("Crosshair", PrimaryOverlays.Last());
-
-	PrimaryOverlays.Add(MakeShareable(new FFilmOverlay_Rabatment));
-	UFilmOverlayToolkit::RegisterPrimaryFilmOverlay("Rabatment", PrimaryOverlays.Last());
-
-	ToggleableOverlays.Add(MakeShareable(new FFilmOverlay_SafeFrame(LOCTEXT("ActionSafe", "Action Safe"), 95.f, FLinearColor::Red)));
-	UFilmOverlayToolkit::RegisterToggleableFilmOverlay("ActionSafe", ToggleableOverlays.Last());
+	auto& ExistingPrimaryOverlays = UFilmOverlayToolkit::GetPrimaryFilmOverlays();
+	
+#define PREPARE_PRIMARY_OVERLAY(Name, Construct) { TSharedPtr<IFilmOverlay> Overlay = ExistingPrimaryOverlays.FindRef(Name); \
+	if (!Overlay) { UFilmOverlayToolkit::RegisterPrimaryFilmOverlay(Name, Construct); }}
+	
+	PREPARE_PRIMARY_OVERLAY(NAME_None, MakeShared<FFilmOverlay_None>());
+	PREPARE_PRIMARY_OVERLAY("Grid3x3", MakeShared<FFilmOverlay_Grid>(3, 3));
+	PREPARE_PRIMARY_OVERLAY("Grid2x2", MakeShared<FFilmOverlay_Grid>(2, 2));
+	PREPARE_PRIMARY_OVERLAY("Crosshair", MakeShared<FFilmOverlay_Crosshair>());
+	PREPARE_PRIMARY_OVERLAY("Rabatment", MakeShared<FFilmOverlay_Rabatment>());
+	
+#undef PREPARE_PRIMARY_OVERLAY
 
-	ToggleableOverlays.Add(MakeShareable(new FFilmOverlay_SafeFrame(LOCTEXT("TitleSafe", "Title Safe"), 90.f, FLinearColor::Yellow)));
-	UFilmOverlayToolkit::RegisterToggleableFilmOverlay("TitleSafe", ToggleableOverlays.Last());
+	auto& ExistingToggleableOverplays = UFilmOverlayToolkit::GetToggleableFilmOverlays();
 
-	ToggleableOverlays.Add(MakeShareable(new FFilmOverlay_SafeFrame(LOCTEXT("CustomSafe", "Custom Safe"), 85.f, FLinearColor::Green)));
-	UFilmOverlayToolkit::RegisterToggleableFilmOverlay("CustomSafe", ToggleableOverlays.Last());
+#define PREPARE_TOGGLEABLE_OVERLAY(Name, Construct) { TSharedPtr<IFilmOverlay> Overlay = ExistingToggleableOverplays.FindRef(Name); \
+	if (!Overlay) { UFilmOverlayToolkit::RegisterToggleableFilmOverlay(Name, Construct); }}
+	
+	PREPARE_TOGGLEABLE_OVERLAY("ActionSafe", MakeShared<FFilmOverlay_SafeFrame>(LOCTEXT("ActionSafe", "Action Safe"), 95.f, FLinearColor::Red));
+	PREPARE_TOGGLEABLE_OVERLAY("TitleSafe", MakeShared<FFilmOverlay_SafeFrame>(LOCTEXT("TitleSafe", "Title Safe"), 90.f, FLinearColor::Yellow));
+	PREPARE_TOGGLEABLE_OVERLAY("CustomSafe", MakeShared<FFilmOverlay_SafeFrame>(LOCTEXT("CustomSafe", "Custom Safe"), 85.f, FLinearColor::Green));
+	PREPARE_TOGGLEABLE_OVERLAY("LetterBox", MakeShared<FFilmOverlay_LetterBox>());
 	
-	ToggleableOverlays.Add(MakeShareable(new FFilmOverlay_LetterBox));
-	UFilmOverlayToolkit::RegisterToggleableFilmOverlay("LetterBox",	ToggleableOverlays.Last());
+#undef PREPARE_TOGGLEABLE_OVERLAY
 
 	// used to choose between creating the whole button + menu widget, or just the menu widget
 	// e.g. Old toolbar uses the combo button, while the new one does not
diff --git a/Engine/Plugins/MovieScene/LevelSequenceEditor/Source/LevelSequenceEditor/Private/CinematicViewport/FilmOverlays.h b/Engine/Plugins/MovieScene/LevelSequenceEditor/Source/LevelSequenceEditor/Private/CinematicViewport/FilmOverlays.h
--- a/Engine/Plugins/MovieScene/LevelSequenceEditor/Source/LevelSequenceEditor/Private/CinematicViewport/FilmOverlays.h
+++ b/Engine/Plugins/MovieScene/LevelSequenceEditor/Source/LevelSequenceEditor/Private/CinematicViewport/FilmOverlays.h
@@ -110,12 +110,6 @@ private:
 	/** Color tint to apply to primary overlays */
 	FLinearColor PrimaryColorTint;
 
-	/** Primary overlays (only one can be active at a time) */
-	TArray<TSharedPtr<IFilmOverlay>> PrimaryOverlays;
-
-	/** Toggleable overlays (any number can be active at a time) */
-	TArray<TSharedPtr<IFilmOverlay>> ToggleableOverlays;
-
 	/** The overlay widget we control - externally owned */
 	TWeakPtr<SFilmOverlay> OverlayWidget;
 };
