Adding "Deselect Binding" sequencer function

diff --git a/Engine/Plugins/MovieScene/LevelSequenceEditor/Source/LevelSequenceEditor/Private/LevelSequenceEditorBlueprintLibrary.cpp b/Engine/Plugins/MovieScene/LevelSequenceEditor/Source/LevelSequenceEditor/Private/LevelSequenceEditorBlueprintLibrary.cpp
--- a/Engine/Plugins/MovieScene/LevelSequenceEditor/Source/LevelSequenceEditor/Private/LevelSequenceEditorBlueprintLibrary.cpp
+++ b/Engine/Plugins/MovieScene/LevelSequenceEditor/Source/LevelSequenceEditor/Private/LevelSequenceEditorBlueprintLibrary.cpp
@@ -613,6 +613,17 @@ void ULevelSequenceEditorBlueprintLibrary::SelectBindings(const TArray<FMovieSce
 	}
 }
 
+void ULevelSequenceEditorBlueprintLibrary::DeselectBindings(const TArray<FMovieSceneBindingProxy>& ObjectBindings)
+{
+	if (CurrentSequencer.IsValid())
+	{
+		for (const FMovieSceneBindingProxy& ObjectBinding : ObjectBindings)
+		{
+			CurrentSequencer.Pin()->DeselectObject(ObjectBinding.BindingID);
+		}
+	}
+}
+
 void ULevelSequenceEditorBlueprintLibrary::EmptySelection()
 {
 	if (CurrentSequencer.IsValid())
diff --git a/Engine/Plugins/MovieScene/LevelSequenceEditor/Source/LevelSequenceEditor/Public/LevelSequenceEditorBlueprintLibrary.h b/Engine/Plugins/MovieScene/LevelSequenceEditor/Source/LevelSequenceEditor/Public/LevelSequenceEditorBlueprintLibrary.h
--- a/Engine/Plugins/MovieScene/LevelSequenceEditor/Source/LevelSequenceEditor/Public/LevelSequenceEditorBlueprintLibrary.h
+++ b/Engine/Plugins/MovieScene/LevelSequenceEditor/Source/LevelSequenceEditor/Public/LevelSequenceEditorBlueprintLibrary.h
@@ -221,6 +221,10 @@ public:
 	UFUNCTION(BlueprintCallable, Category = "Level Sequence Editor")
 	static void SelectBindings(const TArray<FMovieSceneBindingProxy>& ObjectBindings);
 
+	/** Deselect bindings */
+	UFUNCTION(BlueprintCallable, Category = "Level Sequence Editor")
+	static void DeselectBindings(const TArray<FMovieSceneBindingProxy>& ObjectBindings);
+
 	/** Empties the current selection. */
 	UFUNCTION(BlueprintCallable, Category = "Level Sequence Editor")
 	static void EmptySelection();
diff --git a/Engine/Source/Editor/Sequencer/Private/Sequencer.cpp b/Engine/Source/Editor/Sequencer/Private/Sequencer.cpp
--- a/Engine/Source/Editor/Sequencer/Private/Sequencer.cpp
+++ b/Engine/Source/Editor/Sequencer/Private/Sequencer.cpp
@@ -7617,6 +7617,19 @@ void FSequencer::SelectObject(FGuid ObjectBinding)
 	}
 }
 
+void FSequencer::DeselectObject(FGuid ObjectBinding)
+{
+	using namespace UE::Sequencer;
+
+	FObjectBindingModelStorageExtension* ObjectStorage = ViewModel->GetRootModel()->CastDynamic<FObjectBindingModelStorageExtension>();
+	check(ObjectStorage);
+
+	if (TSharedPtr<FObjectBindingModel> Model = ObjectStorage->FindModelForObjectBinding(ObjectBinding))
+	{
+		ViewModel->GetSelection()->Outliner.Deselect(Model);
+	}
+}
+
 void FSequencer::SelectTrack(UMovieSceneTrack* Track)
 {
 	using namespace UE::Sequencer;
diff --git a/Engine/Source/Editor/Sequencer/Private/Sequencer.h b/Engine/Source/Editor/Sequencer/Private/Sequencer.h
--- a/Engine/Source/Editor/Sequencer/Private/Sequencer.h
+++ b/Engine/Source/Editor/Sequencer/Private/Sequencer.h
@@ -794,6 +794,7 @@ public:
 	virtual void GetSelectedKeyAreas(TArray<const IKeyArea*>& OutSelectedKeyAreas, bool bIncludeSelectedKeys = true)  override;
 	virtual void GetSelectedObjects(TArray<FGuid>& OutObjects) override;
 	virtual void SelectObject(FGuid ObjectBinding) override;
+	virtual void DeselectObject(FGuid ObjectBinding) override;
 	virtual void SelectTrack(UMovieSceneTrack* Track) override;
 	virtual void SelectSection(UMovieSceneSection* Section) override;
 	virtual void SelectFolder(UMovieSceneFolder* Folder) override;
diff --git a/Engine/Source/Editor/Sequencer/Public/ISequencer.h b/Engine/Source/Editor/Sequencer/Public/ISequencer.h
--- a/Engine/Source/Editor/Sequencer/Public/ISequencer.h
+++ b/Engine/Source/Editor/Sequencer/Public/ISequencer.h
@@ -702,6 +702,9 @@ public:
 	/** Selects an object by GUID */
 	virtual void SelectObject(FGuid ObjectBinding) = 0;
 
+	/** Deselects an object by GUID */
+	virtual void DeselectObject(FGuid ObjectBinding) = 0;
+
 	/** Selects a track */
 	virtual void SelectTrack(UMovieSceneTrack* Track) = 0;
 
