Sequencer: Move Save as to SSequencer so that any Sequencer (ie. Template Sequence or Camera Animation Sequence) that supports Save also will have a Save as menu item

diff --git a/Engine/Plugins/MovieScene/LevelSequenceEditor/Source/LevelSequenceEditor/Private/Misc/LevelSequenceCustomization.cpp b/Engine/Plugins/MovieScene/LevelSequenceEditor/Source/LevelSequenceEditor/Private/Misc/LevelSequenceCustomization.cpp
--- a/Engine/Plugins/MovieScene/LevelSequenceEditor/Source/LevelSequenceEditor/Private/Misc/LevelSequenceCustomization.cpp
+++ b/Engine/Plugins/MovieScene/LevelSequenceEditor/Source/LevelSequenceEditor/Private/Misc/LevelSequenceCustomization.cpp
@@ -93,23 +93,12 @@ void FLevelSequenceCustomization::ExtendActionsMenu(FMenuBuilder& MenuBuilder)
 	{
 		const FLevelSequenceEditorCommands& Commands = FLevelSequenceEditorCommands::Get();
 		
-		MenuBuilder.AddMenuEntry(
-				LOCTEXT("SaveAs", "Save As..."),
-				LOCTEXT("SaveAsTooltip", "Saves the current sequence under a different name"),
-				FSlateIcon(Commands.GetStyleSetName(), "LevelSequenceEditor.SaveAs"),
-				FUIAction(FExecuteAction::CreateRaw(this, &FLevelSequenceCustomization::OnSaveMovieSceneAsClicked)));
-
 		MenuBuilder.AddMenuEntry(Commands.ImportFBX);
 		MenuBuilder.AddMenuEntry(Commands.ExportFBX);
 	}
 	MenuBuilder.PopCommandList();
 }
 
-void FLevelSequenceCustomization::OnSaveMovieSceneAsClicked()
-{
-	FSequencerUtilities::SaveCurrentMovieSceneAs(WeakSequencer.Pin().ToSharedRef());
-}
-
 void FLevelSequenceCustomization::ImportFBX()
 {
 	FLevelSequenceFBXInterop Interop(WeakSequencer.Pin().ToSharedRef());
diff --git a/Engine/Plugins/MovieScene/LevelSequenceEditor/Source/LevelSequenceEditor/Private/Misc/LevelSequenceCustomization.h b/Engine/Plugins/MovieScene/LevelSequenceEditor/Source/LevelSequenceEditor/Private/Misc/LevelSequenceCustomization.h
--- a/Engine/Plugins/MovieScene/LevelSequenceEditor/Source/LevelSequenceEditor/Private/Misc/LevelSequenceCustomization.h
+++ b/Engine/Plugins/MovieScene/LevelSequenceEditor/Source/LevelSequenceEditor/Private/Misc/LevelSequenceCustomization.h
@@ -30,7 +30,6 @@ private:
 
 	// Actions menu extensions
 	void ExtendActionsMenu(FMenuBuilder& MenuBuilder);
-	void OnSaveMovieSceneAsClicked();
 	void ImportFBX();
 	void ExportFBX();
 
diff --git a/Engine/Source/Editor/Sequencer/Private/SSequencer.cpp b/Engine/Source/Editor/Sequencer/Private/SSequencer.cpp
--- a/Engine/Source/Editor/Sequencer/Private/SSequencer.cpp
+++ b/Engine/Source/Editor/Sequencer/Private/SSequencer.cpp
@@ -1714,6 +1714,15 @@ TSharedRef<SWidget> SSequencer::MakeActionsMenu()
 	TSharedPtr<FExtender> Extender = FExtender::Combine(AllExtenders);
 	FMenuBuilder MenuBuilder(true, Sequencer->GetCommandBindings(), Extender);
 
+	if (Sequencer->GetHostCapabilities().bSupportsSaveMovieSceneAsset)
+	{
+		MenuBuilder.AddMenuEntry(
+			LOCTEXT("SaveAs", "Save As..."),
+			LOCTEXT("SaveAsTooltip", "Saves the current sequence under a different name"),
+			FSlateIcon(FAppStyle::Get().GetStyleSetName(), "Icons.Save"),
+			FUIAction(FExecuteAction::CreateRaw(this, &SSequencer::OnSaveMovieSceneAsClicked)));
+	}
+
 	MenuBuilder.BeginSection("SequenceOptions", LOCTEXT("SequenceOptionsHeader", "Sequence"));
 	{
 		UMovieSceneSequence* RootSequence = Sequencer->GetRootMovieSceneSequence();
@@ -3312,6 +3321,11 @@ void SSequencer::OnSaveMovieSceneClicked()
 	SequencerPtr.Pin()->SaveCurrentMovieScene();
 }
 
+void SSequencer::OnSaveMovieSceneAsClicked()
+{
+	SequencerPtr.Pin()->SaveCurrentMovieSceneAs();
+}
+
 float SSequencer::GetPinnedAreaMaxHeight() const
 {
 	if (!MainSequencerArea.IsValid())
diff --git a/Engine/Source/Editor/Sequencer/Private/SSequencer.h b/Engine/Source/Editor/Sequencer/Private/SSequencer.h
--- a/Engine/Source/Editor/Sequencer/Private/SSequencer.h
+++ b/Engine/Source/Editor/Sequencer/Private/SSequencer.h
@@ -316,6 +316,7 @@ public:
 
 	/** Called when the save button is clicked */
 	void OnSaveMovieSceneClicked();
+	void OnSaveMovieSceneAsClicked();
 
 	/** Called when the curve editor is shown or hidden */
 	void OnCurveEditorVisibilityChanged(bool bShouldBeVisible);
