Sequencer: Fix recursive GetHandleToObject in AddDefaultTracksForActor

diff --git a/Engine/Plugins/MovieScene/LevelSequenceEditor/Source/LevelSequenceEditor/Private/LevelSequenceEditorToolkit.cpp b/Engine/Plugins/MovieScene/LevelSequenceEditor/Source/LevelSequenceEditor/Private/LevelSequenceEditorToolkit.cpp
--- a/Engine/Plugins/MovieScene/LevelSequenceEditor/Source/LevelSequenceEditor/Private/LevelSequenceEditorToolkit.cpp
+++ b/Engine/Plugins/MovieScene/LevelSequenceEditor/Source/LevelSequenceEditor/Private/LevelSequenceEditorToolkit.cpp
@@ -668,7 +668,8 @@ void FLevelSequenceEditorToolkit::AddDefaultTracksForActor(AActor& Actor, const
 				break;
 			}
 
-			FGuid ComponentBinding = Sequencer->GetHandleToObject(PropertyOwner);
+			const bool bCreateHandleIfMissing = false;
+			FGuid ComponentBinding = Sequencer->GetHandleToObject(PropertyOwner, bCreateHandleIfMissing);
 
 			if (bReplaceWithTransformTrack)
 			{
