ShotMangement: Add in Widget to LevelSequence context menu to allow setting of Favorite Rating of a single Level Sequence  and inheriting classes such as Cine Assembly.

diff --git a/Engine/Plugins/MovieScene/LevelSequenceEditor/Source/LevelSequenceEditor/LevelSequenceEditor.Build.cs b/Engine/Plugins/MovieScene/LevelSequenceEditor/Source/LevelSequenceEditor/LevelSequenceEditor.Build.cs
--- a/Engine/Plugins/MovieScene/LevelSequenceEditor/Source/LevelSequenceEditor/LevelSequenceEditor.Build.cs
+++ b/Engine/Plugins/MovieScene/LevelSequenceEditor/Source/LevelSequenceEditor/LevelSequenceEditor.Build.cs
@@ -35,6 +35,7 @@ public class LevelSequenceEditor : ModuleRules
                 "CinematicCamera",
 				"ClassViewer",
 				"Constraints",
+				"ContentBrowser",
 				"Core",
 				"CoreUObject",
 				"CurveEditor",
diff --git a/Engine/Plugins/MovieScene/LevelSequenceEditor/Source/LevelSequenceEditor/Private/LevelSequenceEditorModule.cpp b/Engine/Plugins/MovieScene/LevelSequenceEditor/Source/LevelSequenceEditor/Private/LevelSequenceEditorModule.cpp
--- a/Engine/Plugins/MovieScene/LevelSequenceEditor/Source/LevelSequenceEditor/Private/LevelSequenceEditorModule.cpp
+++ b/Engine/Plugins/MovieScene/LevelSequenceEditor/Source/LevelSequenceEditor/Private/LevelSequenceEditorModule.cpp
@@ -10,6 +10,7 @@
 #include "CameraRig_Rail.h"
 #include "CineCameraActor.h"
 #include "CinematicViewport/CinematicViewportLayoutEntity.h"
+#include "ContentBrowserMenuContexts.h"
 #include "Factories/Factory.h"
 #include "Framework/MultiBox/MultiBoxBuilder.h"
 #include "Framework/MultiBox/MultiBoxExtender.h"
@@ -30,13 +31,56 @@
 #include "SequencerCommands.h"
 #include "SequencerSettings.h"
 #include "Styles/LevelSequenceEditorStyle.h"
+#include "ToolMenus.h"
 #include "UObject/UObjectIterator.h"
+#include "Widgets/SLevelSequenceFavoriteRating.h"
 
 #define LOCTEXT_NAMESPACE "LevelSequenceEditor"
 
 
 TSharedPtr<FLevelSequenceEditorStyle> FLevelSequenceEditorStyle::Singleton;
 
+namespace UE::LevelSequenceEditorModule::Private
+{
+	/** Utility function to make a Menu entry with a custom Slate widget. */
+	FNewToolMenuCustomWidget MakeCustomMenuWidgetDelegate(
+		const TAttribute<FText>& Label, const TAttribute<FSlateIcon>& Icon, TSharedRef<SWidget> CustomWidget, EHorizontalAlignment CustomWidgetHAlign = HAlign_Fill)
+	{
+		return FNewToolMenuCustomWidget::CreateLambda([Label, Icon, CustomWidget, CustomWidgetHAlign](const FToolMenuContext& InContext, const FToolMenuCustomWidgetContext& WidgetContext)
+			{
+				return SNew(SHorizontalBox)
+					+ SHorizontalBox::Slot()
+					.VAlign(VAlign_Center)
+					.Padding(FMargin(2, 0, 0, 0))
+					.AutoWidth()
+					[
+						SNew(SImage)
+							.ColorAndOpacity(FSlateColor::UseSubduedForeground())
+							.Image(TAttribute<const FSlateBrush*>::CreateLambda([Icon]()
+								{
+									return Icon.IsSet() ? Icon.Get().GetIcon() : nullptr;
+								}))
+					]
+					+ SHorizontalBox::Slot()
+						.VAlign(VAlign_Center)
+						.Padding(FMargin(8, 0, 0, 0))
+						.AutoWidth()
+						[
+							SNew(STextBlock)
+								.Text(Label)
+								.TextStyle(WidgetContext.StyleSet, ISlateStyle::Join(WidgetContext.StyleName, ".Label"))
+						]
+						+ SHorizontalBox::Slot()
+						.Padding(FMargin(8, 0, 0, 0))
+						.VAlign(VAlign_Center)
+						.HAlign(CustomWidgetHAlign)
+						.FillWidth(1.0f)
+						[
+							CustomWidget
+						];
+			});
+	};
+}
 
 /**
  * Implements the LevelSequenceEditor module.
@@ -147,6 +191,50 @@ protected:
 
 		FLevelEditorModule& LevelEditorModule = FModuleManager::LoadModuleChecked<FLevelEditorModule>("LevelEditor");
 		LevelEditorModule.GetAllLevelEditorToolbarCinematicsMenuExtenders().Add(CinematicsMenuExtender);
+
+		// Create and register context menu extensions
+		if (UToolMenu* Menu = UE::ContentBrowser::ExtendToolMenu_AssetContextMenu(ULevelSequence::StaticClass()))
+		{
+			const FName AssetActionSectionName = TEXT("GetAssetActions");
+			FToolMenuSection& Section = Menu->FindOrAddSection(AssetActionSectionName);
+
+			const FName FavoriteRatingMenuEntry = "LevelSequenceFavoriteRating";
+			Section.AddDynamicEntry(FavoriteRatingMenuEntry, FNewToolMenuSectionDelegate::CreateLambda([this](FToolMenuSection& InSection)
+			{
+				if (UContentBrowserAssetContextMenuContext* Context = InSection.FindContext<UContentBrowserAssetContextMenuContext>())
+				{
+					// Ensure only a single selected asset.
+					const TArray<FAssetData>& SelectedAssets = Context->SelectedAssets;
+					if (SelectedAssets.Num() == 1)
+					{
+						const FAssetData& LevelSequenceData = SelectedAssets[0];
+						ULevelSequence* LevelSequence = Cast<ULevelSequence>(LevelSequenceData.GetAsset());
+
+						FText SetFavoriteRatingLabel = LOCTEXT("SetFavoriteRating", "Set Rating");
+						FSlateIcon SetFavoriteRatingIcon(FAppStyle::GetAppStyleSetName(), "Icons.Star");
+
+						FToolMenuEntry RatingMenuEntry = FToolMenuEntry::InitMenuEntry(
+							"SetFavoriteRating",
+							SetFavoriteRatingLabel,
+							LOCTEXT("SetFavoriteRatingTooltip", "Set Rating"),
+							SetFavoriteRatingIcon,
+							FUIAction()
+						);
+
+						TSharedRef<SLevelSequenceFavoriteRating> RatingWidget = SNew(SLevelSequenceFavoriteRating)
+							.LevelSequence(LevelSequence);
+
+						RatingMenuEntry.MakeCustomWidget = UE::LevelSequenceEditorModule::Private::MakeCustomMenuWidgetDelegate(
+							SetFavoriteRatingLabel,
+							SetFavoriteRatingIcon,
+							RatingWidget,
+							HAlign_Right
+						);
+						InSection.AddEntry(RatingMenuEntry);
+					}
+				}
+			}));
+		}
 	}
 
 	/** Registers placement mode extensions. */
diff --git a/Engine/Plugins/MovieScene/LevelSequenceEditor/Source/LevelSequenceEditor/Private/Styles/LevelSequenceEditorStyle.h b/Engine/Plugins/MovieScene/LevelSequenceEditor/Source/LevelSequenceEditor/Private/Styles/LevelSequenceEditorStyle.h
--- a/Engine/Plugins/MovieScene/LevelSequenceEditor/Source/LevelSequenceEditor/Private/Styles/LevelSequenceEditorStyle.h
+++ b/Engine/Plugins/MovieScene/LevelSequenceEditor/Source/LevelSequenceEditor/Private/Styles/LevelSequenceEditorStyle.h
@@ -6,6 +6,8 @@
 #include "Brushes/SlateBorderBrush.h"
 #include "Brushes/SlateImageBrush.h"
 #include "Interfaces/IPluginManager.h"
+#include "Styling/SlateTypes.h"
+#include "Styling/StyleColors.h"
 #include "Styling/SlateStyleRegistry.h"
 #include "Styling/SlateStyle.h"
 #include "Styling/SlateStyleMacros.h"
@@ -54,6 +56,9 @@ public:
 		Set("LevelSequenceEditor.ImportFBX", new IMAGE_BRUSH("Icon_Sequencer_ImportFBX_48x", Icon48x48));
 		Set("LevelSequenceEditor.ExportFBX", new IMAGE_BRUSH("Icon_Sequencer_ExportFBX_48x", Icon48x48));
 
+		Set("LevelSequenceEditor.Flag", new IMAGE_BRUSH_SVG("Flag", Icon16x16));
+		Set("LevelSequenceEditor.ThumbsDown", new IMAGE_BRUSH_SVG("ThumbsDown", Icon16x16));
+
 		Set("FilmOverlay.DefaultThumbnail", new IMAGE_BRUSH("DefaultFilmOverlayThumbnail", FVector2D(36, 24)));
 
 	 	Set("FilmOverlay.Disabled", new IMAGE_BRUSH("FilmOverlay.Disabled", FVector2D(36, 24)));
@@ -68,6 +73,52 @@ public:
 	 	Set("FilmOverlay.Crosshair.Small", new IMAGE_BRUSH_SVG("FilmOverlay.Crosshair_16", Icon16x16));
 	 	Set("FilmOverlay.Rabatment.Small", new IMAGE_BRUSH_SVG("FilmOverlay.Rabatment_16", Icon16x16));
 
+		// Widget Styles
+		const FSlateColor NoGoodCheckedColor(FStyleColors::AccentRed);
+		const FSlateColor NoGoodHoveredColor(NoGoodCheckedColor.UseSubduedForeground());
+		const FSlateColor NoGoodPressedColor(NoGoodCheckedColor.GetSpecifiedColor() * 0.8f);
+
+		FCheckBoxStyle NoGoodBoxStyle = FCheckBoxStyle()
+			.SetCheckBoxType(ESlateCheckBoxType::ToggleButton)
+			.SetForegroundColor(FLinearColor::White)
+			.SetCheckedForegroundColor(NoGoodCheckedColor)
+			.SetHoveredForegroundColor(NoGoodHoveredColor)
+			.SetPressedForegroundColor(NoGoodPressedColor)
+			.SetCheckedHoveredForegroundColor(NoGoodHoveredColor)
+			.SetCheckedPressedForegroundColor(NoGoodPressedColor)
+			.SetBorderBackgroundColor(FLinearColor::Transparent);
+
+		Set("NoGoodWidget", NoGoodBoxStyle);
+
+		const FSlateColor FlaggedCheckedColor(FStyleColors::AccentYellow);
+		const FSlateColor FlaggedHoveredColor(FlaggedCheckedColor.UseSubduedForeground());
+		const FSlateColor FlaggedPressedColor(FlaggedCheckedColor.GetSpecifiedColor() * 0.8f);
+
+		FCheckBoxStyle FlaggedBoxStyle = FCheckBoxStyle()
+			.SetCheckBoxType(ESlateCheckBoxType::ToggleButton)
+			.SetForegroundColor(FLinearColor::White)
+			.SetCheckedForegroundColor(FlaggedCheckedColor)
+			.SetHoveredForegroundColor(FlaggedHoveredColor)
+			.SetPressedForegroundColor(FlaggedPressedColor)
+			.SetCheckedHoveredForegroundColor(FlaggedHoveredColor)
+			.SetCheckedPressedForegroundColor(FlaggedPressedColor)
+			.SetBorderBackgroundColor(FLinearColor::Transparent);
+
+		Set("FlaggedWidget", FlaggedBoxStyle);
+
+		FCheckBoxStyle ScoreRatingComponentStyle = FCheckBoxStyle()
+			.SetCheckBoxType(ESlateCheckBoxType::ToggleButton)
+			.SetForegroundColor(FSlateColor(FLinearColor::White))
+			.SetCheckedForegroundColor(FStyleColors::AccentBlue)
+			.SetHoveredForegroundColor(FStyleColors::PrimaryHover)
+			.SetPressedForegroundColor(FStyleColors::PrimaryPress)
+			.SetCheckedHoveredForegroundColor(FStyleColors::PrimaryHover)
+			.SetCheckedPressedForegroundColor(FStyleColors::PrimaryPress)
+			.SetBorderBackgroundColor(FSlateColor(FLinearColor::Transparent))
+			.SetPadding(FMargin(8, 3));;
+
+		Set("ScoreRatingComponent", ScoreRatingComponentStyle);
+
 		FSlateStyleRegistry::RegisterSlateStyle(*this);
 	 }
 
diff --git a/Engine/Plugins/MovieScene/LevelSequenceEditor/Source/LevelSequenceEditor/Private/Widgets/SLevelSequenceFavoriteRating.cpp b/Engine/Plugins/MovieScene/LevelSequenceEditor/Source/LevelSequenceEditor/Private/Widgets/SLevelSequenceFavoriteRating.cpp
new file mode 100644
--- /dev/null
+++ b/Engine/Plugins/MovieScene/LevelSequenceEditor/Source/LevelSequenceEditor/Private/Widgets/SLevelSequenceFavoriteRating.cpp
@@ -0,0 +1,114 @@
+// Copyright Epic Games, Inc. All Rights Reserved.
+
+#include "Widgets/SLevelSequenceFavoriteRating.h"
+#include "MetaData/MovieSceneShotMetaData.h"
+#include "Styles/LevelSequenceEditorStyle.h"
+#include "Styling/StyleColors.h"
+#include "Widgets/SScoreRating.h"
+#include "Widgets/Images/SImage.h"
+#include "Widgets/Input/SCheckBox.h"
+
+#define LOCTEXT_NAMESPACE "SLevelSequenceFavoriteRating"
+
+void SLevelSequenceFavoriteRating::Construct(const FArguments& InArgs)
+{
+	WeakLevelSequence = InArgs._LevelSequence;
+
+	ChildSlot
+		[
+			SNew(SHorizontalBox)
+
+				// Thumbs down
+				+ SHorizontalBox::Slot()
+					.AutoWidth()
+					[
+						SNew(SCheckBox)
+							.Padding(FMargin(0.0, 0.0, 8.0, 0.0))
+							.Style(&FLevelSequenceEditorStyle::Get()->GetWidgetStyle<FCheckBoxStyle>("NoGoodWidget"))
+							.ToolTipText(LOCTEXT("MarkLevelSequenceAsNoGoodTooltip", "Mark as No Good"))
+							.IsChecked_Lambda([&WeakLevelSequence = this->WeakLevelSequence]()->ECheckBoxState
+								{
+									bool bValue = false;
+									if (UMovieSceneShotMetaData* MetaData = WeakLevelSequence.IsValid() ? WeakLevelSequence.Pin()->FindMetaData<UMovieSceneShotMetaData>() : nullptr)
+									{
+										bValue = MetaData->GetIsNoGood().IsSet() ? MetaData->GetIsNoGood().GetValue() : false;
+									}
+									return bValue ? ECheckBoxState::Checked : ECheckBoxState::Unchecked;
+								})
+							.OnCheckStateChanged(FOnCheckStateChanged::CreateLambda([&WeakLevelSequence = this->WeakLevelSequence](ECheckBoxState State)
+								{
+									if (UMovieSceneShotMetaData* MetaData = WeakLevelSequence.IsValid() ? WeakLevelSequence.Pin()->FindOrAddMetaData<UMovieSceneShotMetaData>() : nullptr)
+									{
+										MetaData->SetIsNoGood(State == ECheckBoxState::Checked);
+									}
+								}))
+							[
+								SNew(SImage)
+									.ColorAndOpacity(FSlateColor::UseForeground())
+									.Image(FLevelSequenceEditorStyle::Get()->GetBrush("LevelSequenceEditor.ThumbsDown"))
+							]
+					]
+
+				// Flagged
+				+ SHorizontalBox::Slot()
+					.AutoWidth()
+					[
+						SNew(SCheckBox)
+							.Padding(FMargin(8.0, 0.0, 16.0, 0.0))
+							.Style(&FLevelSequenceEditorStyle::Get()->GetWidgetStyle<FCheckBoxStyle>("FlaggedWidget"))
+							.ToolTipText(LOCTEXT("FlagLevelSequenceTooltip", "Flag"))
+							.IsChecked_Lambda([&WeakLevelSequence = this->WeakLevelSequence]()->ECheckBoxState
+								{
+									bool bValue = false;
+									if (UMovieSceneShotMetaData* MetaData = WeakLevelSequence.IsValid() ? WeakLevelSequence.Pin()->FindMetaData<UMovieSceneShotMetaData>() : nullptr)
+									{
+										bValue = MetaData->GetIsFlagged().IsSet() ? MetaData->GetIsFlagged().GetValue() : false;
+									}
+									return bValue ? ECheckBoxState::Checked : ECheckBoxState::Unchecked;
+								})
+							.OnCheckStateChanged(FOnCheckStateChanged::CreateLambda([&WeakLevelSequence = this->WeakLevelSequence](ECheckBoxState State)
+								{
+									if (UMovieSceneShotMetaData* MetaData = WeakLevelSequence.IsValid() ? WeakLevelSequence.Pin()->FindOrAddMetaData<UMovieSceneShotMetaData>() : nullptr)
+									{
+										MetaData->SetIsFlagged(State == ECheckBoxState::Checked);
+									}
+								}))
+							[
+								SNew(SImage)
+									.ColorAndOpacity(FSlateColor::UseForeground())
+									.Image(FLevelSequenceEditorStyle::Get()->GetBrush("LevelSequenceEditor.Flag"))
+							]
+					]
+
+				// Rating
+				+ SHorizontalBox::Slot()
+					.AutoWidth()
+					.HAlign(HAlign_Right)
+					[
+						SNew(SScoreRating)
+							.MaxScore(3)
+							.OnScoreChanged_Lambda([&WeakLevelSequence = this->WeakLevelSequence](int NewScore)
+								{
+									if (UMovieSceneShotMetaData* MetaData = WeakLevelSequence.IsValid() ? WeakLevelSequence.Pin()->FindOrAddMetaData<UMovieSceneShotMetaData>() : nullptr)
+									{
+										if (NewScore > 0)
+										{
+											MetaData->SetFavoriteRating(NewScore);
+										}
+										else
+										{
+											MetaData->ClearFavoriteRating();
+										}
+									}
+								})
+							.Score_Lambda([&WeakLevelSequence = this->WeakLevelSequence]() -> int
+								{
+									UMovieSceneShotMetaData* MetaData = WeakLevelSequence.IsValid() ? WeakLevelSequence.Pin()->FindMetaData<UMovieSceneShotMetaData>() : nullptr;
+									return MetaData && MetaData->GetFavoriteRating().IsSet() ? MetaData->GetFavoriteRating().GetValue() : 0;
+								})
+
+					]
+		];
+}
+
+#undef LOCTEXT_NAMESPACE
diff --git a/Engine/Plugins/MovieScene/LevelSequenceEditor/Source/LevelSequenceEditor/Private/Widgets/SScoreRating.cpp b/Engine/Plugins/MovieScene/LevelSequenceEditor/Source/LevelSequenceEditor/Private/Widgets/SScoreRating.cpp
new file mode 100644
--- /dev/null
+++ b/Engine/Plugins/MovieScene/LevelSequenceEditor/Source/LevelSequenceEditor/Private/Widgets/SScoreRating.cpp
@@ -0,0 +1,70 @@
+// Copyright Epic Games, Inc. All Rights Reserved.
+
+#include "Widgets/SScoreRating.h"
+#include "Widgets/SScoreRatingComponent.h"
+
+#define LOCTEXT_NAMESPACE "SScoreRating"
+
+SLATE_IMPLEMENT_WIDGET(SScoreRating)
+void SScoreRating::PrivateRegisterAttributes(FSlateAttributeInitializer& AttributeInitializer)
+{
+	SLATE_ADD_MEMBER_ATTRIBUTE_DEFINITION_WITH_NAME(AttributeInitializer, "ComponentIconAttribute", ComponentIconAttribute, EInvalidateWidgetReason::Layout);
+	SLATE_ADD_MEMBER_ATTRIBUTE_DEFINITION_WITH_NAME(AttributeInitializer, "ScoreAttribute", ScoreAttribute, EInvalidateWidgetReason::Layout)
+	.OnValueChanged(FSlateAttributeDescriptor::FAttributeValueChangedDelegate::CreateLambda([](SWidget& Widget)
+		{
+			static_cast<SScoreRating&>(Widget).NotifyScoreChanged();
+		}
+	));
+}
+
+SScoreRating::SScoreRating() :
+	ComponentIconAttribute(*this),
+	ScoreAttribute(*this)
+{
+
+}
+
+void SScoreRating::Construct(const FArguments& InArgs)
+{
+	LayoutBox = SNew(SHorizontalBox);
+	// Ensure we always have at least 2 and no more than 10
+	MaxScore = FMath::Clamp(InArgs._MaxScore, 2, 10);
+	OnScoreChanged = InArgs._OnScoreChanged;
+	ScoreAttribute.Assign(*this, InArgs._Score);
+	ComponentIconAttribute.Assign(*this, InArgs._ComponentIcon);
+
+	for (int ScoreIndex = 0; ScoreIndex < MaxScore; ++ScoreIndex)
+	{
+		LayoutBox->AddSlot()
+			.AutoWidth()
+			.Padding(FMargin(0.0, 0.0, 4.0, 0.0))
+			[
+				SNew(SScoreRatingComponent)
+					.Icon(InArgs._ComponentIcon)
+					.ToolTipText(FText::Format(LOCTEXT("SetScoreRatingTooltip", "Score as {0}"), ScoreIndex + 1))
+					.IsChecked_Lambda([&Score = this->ScoreAttribute, &HoverIndex = this->HoverIndex, ScoreIndex]() -> bool
+						{
+							return  ScoreIndex < HoverIndex || Score.Get() > ScoreIndex;
+						})
+					.OnCheckStateChanged(FOnBinaryCheckStateChanged::CreateLambda([this, ScoreIndex](bool bIsChecked)
+						{
+							const int NewScore = ScoreIndex + 1;
+							ScoreAttribute.Set(*this, NewScore != ScoreAttribute.Get() ? NewScore : 0);
+						}))
+					.OnHoverStarted_Raw(this, &SScoreRating::SetHoverIndex, ScoreIndex)
+					.OnHoverFinished_Raw(this, &SScoreRating::SetHoverIndex, static_cast<int>(INDEX_NONE))
+			];
+	}
+
+	ChildSlot
+		[
+			LayoutBox.ToSharedRef()
+		];
+}
+
+void SScoreRating::NotifyScoreChanged()
+{
+	OnScoreChanged.ExecuteIfBound(ScoreAttribute.Get());
+}
+
+#undef LOCTEXT_NAMESPACE
diff --git a/Engine/Plugins/MovieScene/LevelSequenceEditor/Source/LevelSequenceEditor/Private/Widgets/SScoreRating.h b/Engine/Plugins/MovieScene/LevelSequenceEditor/Source/LevelSequenceEditor/Private/Widgets/SScoreRating.h
new file mode 100644
--- /dev/null
+++ b/Engine/Plugins/MovieScene/LevelSequenceEditor/Source/LevelSequenceEditor/Private/Widgets/SScoreRating.h
@@ -0,0 +1,67 @@
+// Copyright Epic Games, Inc. All Rights Reserved.
+
+#pragma once
+
+#include "Styling/AppStyle.h"
+#include "Widgets/SCompoundWidget.h"
+#include "Widgets/DeclarativeSyntaxSupport.h"
+#include "Widgets/SBoxPanel.h"
+
+/** Delegate that is executed when the score is changed */
+DECLARE_DELEGATE_OneParam(FOnScoreChanged, int )
+
+/** Widget for viewing and setting a score rating. */
+class SScoreRating : public SCompoundWidget
+{
+	SLATE_DECLARE_WIDGET(SScoreRating, SCompoundWidget)
+
+public:
+	SLATE_BEGIN_ARGS(SScoreRating)
+		: _MaxScore(3)
+		, _ComponentIcon(FAppStyle::GetBrush("Icons.Star"))
+		, _OnScoreChanged()
+		, _Score(0)
+		{}
+
+		/** The max score shown by the widget */
+		SLATE_ARGUMENT( int, MaxScore )
+
+		/** The icon brush to use for each score rating component */
+		SLATE_ATTRIBUTE( const FSlateBrush*, ComponentIcon )
+
+		/** Event triggered when the score is changed. */
+		SLATE_EVENT( FOnScoreChanged, OnScoreChanged)
+
+		/** The current score */
+		SLATE_ATTRIBUTE( int, Score )
+	SLATE_END_ARGS()
+
+	SScoreRating();
+
+	void Construct(const FArguments& InArgs);
+
+private:
+	/** Set the current hover index */
+	inline void SetHoverIndex(int Index) { HoverIndex = Index; }
+
+	/** Notify the score has changed with the new score */
+	void NotifyScoreChanged();
+
+	/** The overall layout widget */
+	TSharedPtr<SHorizontalBox> LayoutBox;
+
+	/** Max possible score to display */
+	int MaxScore = 3;
+
+	/** Event triggered when the score is changed. */
+	FOnScoreChanged OnScoreChanged;
+
+	/** Current score rating index that the mouse is hovering over */
+	int HoverIndex = 0;
+
+	/** The slate brush to use to draw the component icons. */
+	TSlateAttribute<const FSlateBrush*> ComponentIconAttribute;
+
+	/** The current score */
+	TSlateAttribute<int> ScoreAttribute;
+};
diff --git a/Engine/Plugins/MovieScene/LevelSequenceEditor/Source/LevelSequenceEditor/Private/Widgets/SScoreRatingComponent.cpp b/Engine/Plugins/MovieScene/LevelSequenceEditor/Source/LevelSequenceEditor/Private/Widgets/SScoreRatingComponent.cpp
new file mode 100644
--- /dev/null
+++ b/Engine/Plugins/MovieScene/LevelSequenceEditor/Source/LevelSequenceEditor/Private/Widgets/SScoreRatingComponent.cpp
@@ -0,0 +1,62 @@
+// Copyright Epic Games, Inc. All Rights Reserved.
+
+#include "Widgets/SScoreRatingComponent.h"
+#include "Styles/LevelSequenceEditorStyle.h"
+#include "Widgets/Images/SImage.h"
+#include "Widgets/Input/SCheckBox.h"
+
+#define LOCTEXT_NAMESPACE "SScoreRatingComponent"
+
+SLATE_IMPLEMENT_WIDGET(SScoreRatingComponent)
+void SScoreRatingComponent::PrivateRegisterAttributes(FSlateAttributeInitializer& AttributeInitializer)
+{
+	SLATE_ADD_MEMBER_ATTRIBUTE_DEFINITION_WITH_NAME(AttributeInitializer, "IconAttribute", IconAttribute, EInvalidateWidgetReason::Layout);
+	SLATE_ADD_MEMBER_ATTRIBUTE_DEFINITION_WITH_NAME(AttributeInitializer, "IsCheckedAttribute", IsCheckedAttribute, EInvalidateWidgetReason::Layout);
+}
+
+SScoreRatingComponent::SScoreRatingComponent() :
+	IconAttribute(*this),
+	IsCheckedAttribute(*this)
+{
+
+}
+
+void SScoreRatingComponent::Construct(const FArguments& InArgs)
+{
+	OnHoverStarted = InArgs._OnHoverStarted;
+	OnHoverFinished = InArgs._OnHoverFinished;
+	OnCheckStateChanged = InArgs._OnCheckStateChanged;
+	IconAttribute.Assign(*this, InArgs._Icon);
+	IsCheckedAttribute.Assign(*this, InArgs._IsChecked);
+
+	ChildSlot
+		[
+			SNew(SCheckBox)
+				.Padding(FMargin(0.0f))
+				.Style(&FLevelSequenceEditorStyle::Get()->GetWidgetStyle<FCheckBoxStyle>("ScoreRatingComponent"))
+				.IsChecked_Lambda([&IsChecked = this->IsCheckedAttribute]() -> ECheckBoxState { return IsChecked.Get() ? ECheckBoxState::Checked : ECheckBoxState::Unchecked; })
+				.OnCheckStateChanged(FOnCheckStateChanged::CreateLambda([this, &OnCheckStateChanged = this->OnCheckStateChanged](ECheckBoxState State) 
+					{
+						OnCheckStateChanged.ExecuteIfBound(State == ECheckBoxState::Checked); 
+					}))
+				[
+					SNew(SImage)
+						.ColorAndOpacity(FSlateColor::UseForeground())
+						.Image(InArgs._Icon.Get())
+				]
+		];
+}
+
+void SScoreRatingComponent::OnMouseEnter(const FGeometry& MyGeometry, const FPointerEvent& MouseEvent)
+{
+	SCompoundWidget::OnMouseEnter(MyGeometry, MouseEvent);
+	OnHoverStarted.ExecuteIfBound();
+}
+
+void SScoreRatingComponent::OnMouseLeave(const FPointerEvent& MouseEvent)
+{
+	SCompoundWidget::OnMouseLeave(MouseEvent);
+	OnHoverFinished.ExecuteIfBound();
+}
+
+#undef LOCTEXT_NAMESPACE
diff --git a/Engine/Plugins/MovieScene/LevelSequenceEditor/Source/LevelSequenceEditor/Private/Widgets/SScoreRatingComponent.h b/Engine/Plugins/MovieScene/LevelSequenceEditor/Source/LevelSequenceEditor/Private/Widgets/SScoreRatingComponent.h
new file mode 100644
--- /dev/null
+++ b/Engine/Plugins/MovieScene/LevelSequenceEditor/Source/LevelSequenceEditor/Private/Widgets/SScoreRatingComponent.h
@@ -0,0 +1,74 @@
+// Copyright Epic Games, Inc. All Rights Reserved.
+
+#pragma once
+
+#include "Styling/AppStyle.h"
+#include "Widgets/SCompoundWidget.h"
+#include "Widgets/DeclarativeSyntaxSupport.h"
+
+/** Delegate that is executed when a Score Rating Component hover event happens */
+DECLARE_DELEGATE(FOnScoreRatingComponentHoverEvent);
+
+/** Delegate that is executed when a binary check state change occurs, either to/from true/false */
+DECLARE_DELEGATE_OneParam(FOnBinaryCheckStateChanged, bool);
+
+/** 
+ * Specialized widget for a score value to expose mouse hover events occur and to contain the style.
+ * Ideally this would be contained entirely in SScoreRating however the need to get the hover events
+ * mean we require an extra wrapper class.
+ */
+class SScoreRatingComponent : public SCompoundWidget
+{
+	SLATE_DECLARE_WIDGET(SScoreRatingComponent, SCompoundWidget)
+
+public:
+	SLATE_BEGIN_ARGS(SScoreRatingComponent)
+		: _OnHoverStarted()
+		, _OnHoverFinished()
+		, _OnCheckStateChanged()
+		, _Icon(FAppStyle::GetBrush("Icons.Star"))
+		, _IsChecked(false)
+		{
+		}
+
+		/** Called when a hover event starts */
+		SLATE_EVENT( FOnScoreRatingComponentHoverEvent, OnHoverStarted )
+
+		/** Called when a hover event finishes */
+		SLATE_EVENT( FOnScoreRatingComponentHoverEvent, OnHoverFinished )
+		
+		/** Called when a a state change occurs on the Rating Star through user input */
+		SLATE_EVENT( FOnBinaryCheckStateChanged, OnCheckStateChanged )
+
+		/** The Slate brush to use for this component */
+		SLATE_ATTRIBUTE( const FSlateBrush*, Icon )
+
+		/** Whether the Score Rating Component is currently in its checked state. */
+		SLATE_ATTRIBUTE( bool, IsChecked )
+
+	SLATE_END_ARGS()
+
+	SScoreRatingComponent();
+
+	void Construct(const FArguments& InArgs);
+
+	// SWidget Interface
+	void OnMouseEnter(const FGeometry& MyGeometry, const FPointerEvent& MouseEvent) override;
+	void OnMouseLeave(const FPointerEvent& MouseEvent) override;
+
+private:
+	/** Delegate fired when a hover event starts */
+	FOnScoreRatingComponentHoverEvent OnHoverStarted;
+
+	/** Delegate fired when a hover event finishes */
+	FOnScoreRatingComponentHoverEvent OnHoverFinished;
+
+	/** Delegate fired when a state change event occurs */
+	FOnBinaryCheckStateChanged OnCheckStateChanged;
+
+	/** The slate brush to use to draw the icon. */
+	TSlateAttribute<const FSlateBrush*> IconAttribute;
+
+	/** Are we checked */
+	TSlateAttribute<bool> IsCheckedAttribute;
+};
diff --git a/Engine/Plugins/MovieScene/LevelSequenceEditor/Source/LevelSequenceEditor/Public/Widgets/SLevelSequenceFavoriteRating.h b/Engine/Plugins/MovieScene/LevelSequenceEditor/Source/LevelSequenceEditor/Public/Widgets/SLevelSequenceFavoriteRating.h
new file mode 100644
--- /dev/null
+++ b/Engine/Plugins/MovieScene/LevelSequenceEditor/Source/LevelSequenceEditor/Public/Widgets/SLevelSequenceFavoriteRating.h
@@ -0,0 +1,32 @@
+// Copyright Epic Games, Inc. All Rights Reserved.
+
+#pragma once
+
+#include "LevelSequence.h"
+#include "Widgets/SCompoundWidget.h"
+#include "Widgets/DeclarativeSyntaxSupport.h"
+
+#define UE_API LEVELSEQUENCEEDITOR_API
+
+/** Widget for viewing and setting the Favorite Rating on LevelSequences. */
+class SLevelSequenceFavoriteRating : public SCompoundWidget
+{
+public:
+	SLATE_BEGIN_ARGS(SLevelSequenceFavoriteRating) 
+		: _LevelSequence(nullptr)
+		{}
+
+		/** The LevelSequence object to focus the widget on */
+		SLATE_ARGUMENT( TWeakObjectPtr<ULevelSequence>, LevelSequence )
+	SLATE_END_ARGS()
+
+	UE_API void Construct(const FArguments& InArgs);
+
+private:
+	/** The LevelSequence object the widget is focused on. Only a single sequence can be focused. 
+	 *  This could be changed using the Undetermined state and special multi icons. 
+	 */
+	TWeakObjectPtr<ULevelSequence> WeakLevelSequence;
+};
+
+#undef UE_API
diff --git a/Engine/Plugins/VirtualProduction/CinematicAssemblyTools/Source/CineAssemblyToolsEditor/Private/CineAssemblyToolsEditorModule.cpp b/Engine/Plugins/VirtualProduction/CinematicAssemblyTools/Source/CineAssemblyToolsEditor/Private/CineAssemblyToolsEditorModule.cpp
--- a/Engine/Plugins/VirtualProduction/CinematicAssemblyTools/Source/CineAssemblyToolsEditor/Private/CineAssemblyToolsEditorModule.cpp
+++ b/Engine/Plugins/VirtualProduction/CinematicAssemblyTools/Source/CineAssemblyToolsEditor/Private/CineAssemblyToolsEditorModule.cpp
@@ -289,13 +289,13 @@ void FCineAssemblyToolsEditorModule::RegisterMenus()
 	UToolMenus* ToolMenus = UToolMenus::Get();
 
 	// Add Content Browser context menu actions to Cine Assembly assets
-	const FName ContextMenuName = TEXT("ContentBrowser.AssetContextMenu.CineAssembly");
+	const FName ContextMenuName = TEXT("ContentBrowser.AssetContextMenu.LevelSequence");
 	if (UToolMenu* Menu = ToolMenus->ExtendMenu(ContextMenuName))
 	{
 		const FName AssetActionSectionName = TEXT("GetAssetActions");
 		FToolMenuSection& Section = Menu->FindOrAddSection(AssetActionSectionName);
 
-		Section.AddDynamicEntry("CineAssemblyActions", FNewToolMenuSectionDelegate::CreateLambda([this](FToolMenuSection& InSection)
+		FToolMenuEntry& CineAssemblyActionsEntry = Section.AddDynamicEntry("CineAssemblyActions", FNewToolMenuSectionDelegate::CreateLambda([this](FToolMenuSection& InSection)
 		{
 			if (UContentBrowserAssetContextMenuContext* Context = InSection.FindContext<UContentBrowserAssetContextMenuContext>())
 			{
@@ -356,6 +356,11 @@ void FCineAssemblyToolsEditorModule::RegisterMenus()
 				}
 			}
 		}));
+
+		// Ensure menu order by forcing our entry after the ones added by TakeRecorder even if the load order differs.
+		const FName RecordIntoTakeRecorderEntry = "RecordIntoTakeRecorder_Label";
+		CineAssemblyActionsEntry.InsertPosition.Name = RecordIntoTakeRecorderEntry;
+		CineAssemblyActionsEntry.InsertPosition.Position = EToolMenuInsertType::After;
 	}
 }
 
