ShotManagement: Fix shadowed capture in lambda

#rnx

diff --git a/Engine/Plugins/MovieScene/LevelSequenceEditor/Source/LevelSequenceEditor/Private/Widgets/SScoreRatingComponent.cpp b/Engine/Plugins/MovieScene/LevelSequenceEditor/Source/LevelSequenceEditor/Private/Widgets/SScoreRatingComponent.cpp
--- a/Engine/Plugins/MovieScene/LevelSequenceEditor/Source/LevelSequenceEditor/Private/Widgets/SScoreRatingComponent.cpp
+++ b/Engine/Plugins/MovieScene/LevelSequenceEditor/Source/LevelSequenceEditor/Private/Widgets/SScoreRatingComponent.cpp
@@ -35,9 +35,9 @@ void SScoreRatingComponent::Construct(const FArguments& InArgs)
 				.Padding(FMargin(0.0f))
 				.Style(&FLevelSequenceEditorStyle::Get()->GetWidgetStyle<FCheckBoxStyle>("ScoreRatingComponent"))
 				.IsChecked_Lambda([&IsChecked = this->IsCheckedAttribute]() -> ECheckBoxState { return IsChecked.Get() ? ECheckBoxState::Checked : ECheckBoxState::Unchecked; })
-				.OnCheckStateChanged(FOnCheckStateChanged::CreateLambda([this, &OnCheckStateChanged = this->OnCheckStateChanged](ECheckBoxState State) 
+				.OnCheckStateChanged(FOnCheckStateChanged::CreateLambda([&OnCheckStateChangedRef = this->OnCheckStateChanged](ECheckBoxState State) 
 					{
-						OnCheckStateChanged.ExecuteIfBound(State == ECheckBoxState::Checked); 
+						OnCheckStateChangedRef.ExecuteIfBound(State == ECheckBoxState::Checked);
 					}))
 				[
 					SNew(SImage)
