PR #13402:  [MovieScene] Exposed FLevelSequencePlaybackContext to other modules

diff --git a/Engine/Plugins/MovieScene/LevelSequenceEditor/Source/LevelSequenceEditor/Private/Misc/LevelSequencePlaybackContext.h b/Engine/Plugins/MovieScene/LevelSequenceEditor/Source/LevelSequenceEditor/Private/Misc/LevelSequencePlaybackContext.h
deleted file mode 100644
diff --git a/Engine/Plugins/MovieScene/LevelSequenceEditor/Source/LevelSequenceEditor/Public/Misc/LevelSequencePlaybackContext.h b/Engine/Plugins/MovieScene/LevelSequenceEditor/Source/LevelSequenceEditor/Public/Misc/LevelSequencePlaybackContext.h
new file mode 100644
--- /dev/null
+++ b/Engine/Plugins/MovieScene/LevelSequenceEditor/Source/LevelSequenceEditor/Public/Misc/LevelSequencePlaybackContext.h
@@ -0,0 +1,101 @@
+// Copyright Epic Games, Inc. All Rights Reserved.
+
+#pragma once
+
+#include "UObject/WeakObjectPtr.h"
+
+class SWidget;
+
+class ALevelSequenceActor;
+class IMovieScenePlaybackClient;
+class ULevelSequence;
+
+#define UE_API LEVELSEQUENCEEDITOR_API
+
+/**
+ * Class that manages the current playback context that a level-sequence editor should use for playback
+ */
+class LEVELSEQUENCEEDITOR_API FLevelSequencePlaybackContext : public TSharedFromThis<FLevelSequencePlaybackContext>
+{
+public:
+
+	FLevelSequencePlaybackContext(ULevelSequence* InLevelSequence);
+	~FLevelSequencePlaybackContext();
+
+	/**
+	 * Gets the level sequence for which we are trying to find the context.
+	 */
+	ULevelSequence* GetLevelSequence() const;
+
+	/**
+	 * Build a world picker widget that allows the user to choose a world, and exit the auto-bind settings
+	 */
+	TSharedRef<SWidget> BuildWorldPickerCombo();
+
+	/**
+	 * Resolve the current world context pointer. Can never be nullptr.
+	 */
+	UObject* GetPlaybackContext() const;
+
+	/**
+	 * Returns GetPlaybackContext as a plain object.
+	 */
+	UObject* GetPlaybackContextAsObject() const;
+
+	/**
+	 * Resolve the current playback client. May be nullptr.
+	 */
+	ALevelSequenceActor* GetPlaybackClient() const;
+
+	/**
+	 * Returns GetPlaybackClient as an interface pointer.
+	 */
+	IMovieScenePlaybackClient* GetPlaybackClientAsInterface() const;
+
+	/**
+	 * Retrieve all the event contexts for the current world
+	 */
+	TArray<UObject*> GetEventContexts() const;
+
+	/**
+	 * Specify a new world to use as the context. Persists until the next PIE or map change event.
+	 * May be null, in which case the context will be recomputed automatically
+	 */
+	void OverrideWith(UWorld* InNewContext, ALevelSequenceActor* InNewClient);
+
+private:
+
+	using FContextAndClient = TTuple<UWorld*, ALevelSequenceActor*>;
+
+	/**
+	 * Compute the new playback context based on the user's current auto-bind settings.
+	 * Will use the first encountered PIE or Simulate world if possible, else the Editor world as a fallback
+	 */
+	static FContextAndClient ComputePlaybackContextAndClient(const ULevelSequence* InLevelSequence);
+
+	/**
+	 * Update the cached context and client pointers if needed.
+	 */
+	void UpdateCachedContextAndClient() const;
+
+	/**
+	 * Gets both the context and client.
+	 */
+	//TTuple<UWorld*, ALevelSequenceActor*>
+	FContextAndClient GetPlaybackContextAndClient() const;
+
+	void OnPieEvent(bool);
+	void OnMapChange(uint32);
+	void OnWorldListChanged(UWorld*);
+
+private:
+
+	/** Level sequence that we should find a context for */
+	TWeakObjectPtr<ULevelSequence> LevelSequence;
+
+	/** Mutable cached context pointer */
+	mutable TWeakObjectPtr<UWorld> WeakCurrentContext;
+
+	/** Mutable cached client pointer */
+	mutable TWeakObjectPtr<ALevelSequenceActor> WeakCurrentClient;
+};
