Sequencer: Fix CIS

diff --git a/Engine/Plugins/MovieScene/LevelSequenceEditor/Source/LevelSequenceEditor/Public/Misc/LevelSequencePlaybackContext.h b/Engine/Plugins/MovieScene/LevelSequenceEditor/Source/LevelSequenceEditor/Public/Misc/LevelSequencePlaybackContext.h
--- a/Engine/Plugins/MovieScene/LevelSequenceEditor/Source/LevelSequenceEditor/Public/Misc/LevelSequencePlaybackContext.h
+++ b/Engine/Plugins/MovieScene/LevelSequenceEditor/Source/LevelSequenceEditor/Public/Misc/LevelSequencePlaybackContext.h
@@ -99,3 +99,5 @@ private:
 	/** Mutable cached client pointer */
 	mutable TWeakObjectPtr<ALevelSequenceActor> WeakCurrentClient;
 };
+
+#undef UE_API
\ No newline at end of file
