Sequencer: Follow up to 44036625

#rnx

diff --git a/Engine/Plugins/MovieScene/LevelSequenceEditor/Source/LevelSequenceEditor/Public/Misc/LevelSequencePlaybackContext.h b/Engine/Plugins/MovieScene/LevelSequenceEditor/Source/LevelSequenceEditor/Public/Misc/LevelSequencePlaybackContext.h
--- a/Engine/Plugins/MovieScene/LevelSequenceEditor/Source/LevelSequenceEditor/Public/Misc/LevelSequencePlaybackContext.h
+++ b/Engine/Plugins/MovieScene/LevelSequenceEditor/Source/LevelSequenceEditor/Public/Misc/LevelSequencePlaybackContext.h
@@ -15,7 +15,7 @@ class ULevelSequence;
 /**
  * Class that manages the current playback context that a level-sequence editor should use for playback
  */
-class LEVELSEQUENCEEDITOR_API FLevelSequencePlaybackContext : public TSharedFromThis<FLevelSequencePlaybackContext>
+class UE_API FLevelSequencePlaybackContext : public TSharedFromThis<FLevelSequencePlaybackContext>
 {
 public:
 
