ShotManagement: Correcting int types to int32 in FavoriteRating widget.

#rnx

diff --git a/Engine/Plugins/MovieScene/LevelSequenceEditor/Source/LevelSequenceEditor/Private/Widgets/SLevelSequenceFavoriteRating.cpp b/Engine/Plugins/MovieScene/LevelSequenceEditor/Source/LevelSequenceEditor/Private/Widgets/SLevelSequenceFavoriteRating.cpp
--- a/Engine/Plugins/MovieScene/LevelSequenceEditor/Source/LevelSequenceEditor/Private/Widgets/SLevelSequenceFavoriteRating.cpp
+++ b/Engine/Plugins/MovieScene/LevelSequenceEditor/Source/LevelSequenceEditor/Private/Widgets/SLevelSequenceFavoriteRating.cpp
@@ -87,7 +87,7 @@ void SLevelSequenceFavoriteRating::Construct(const FArguments& InArgs)
 					[
 						SNew(SScoreRating)
 							.MaxScore(3)
-							.OnScoreChanged_Lambda([&WeakLevelSequence = this->WeakLevelSequence](int NewScore)
+							.OnScoreChanged_Lambda([&WeakLevelSequence = this->WeakLevelSequence](int32 NewScore)
 								{
 									if (UMovieSceneShotMetaData* MetaData = WeakLevelSequence.IsValid() ? WeakLevelSequence.Pin()->FindOrAddMetaData<UMovieSceneShotMetaData>() : nullptr)
 									{
@@ -101,7 +101,7 @@ void SLevelSequenceFavoriteRating::Construct(const FArguments& InArgs)
 										}
 									}
 								})
-							.Score_Lambda([&WeakLevelSequence = this->WeakLevelSequence]() -> int
+							.Score_Lambda([&WeakLevelSequence = this->WeakLevelSequence]() -> int32
 								{
 									UMovieSceneShotMetaData* MetaData = WeakLevelSequence.IsValid() ? WeakLevelSequence.Pin()->FindMetaData<UMovieSceneShotMetaData>() : nullptr;
 									return MetaData && MetaData->GetFavoriteRating().IsSet() ? MetaData->GetFavoriteRating().GetValue() : 0;
diff --git a/Engine/Plugins/MovieScene/LevelSequenceEditor/Source/LevelSequenceEditor/Private/Widgets/SScoreRating.cpp b/Engine/Plugins/MovieScene/LevelSequenceEditor/Source/LevelSequenceEditor/Private/Widgets/SScoreRating.cpp
--- a/Engine/Plugins/MovieScene/LevelSequenceEditor/Source/LevelSequenceEditor/Private/Widgets/SScoreRating.cpp
+++ b/Engine/Plugins/MovieScene/LevelSequenceEditor/Source/LevelSequenceEditor/Private/Widgets/SScoreRating.cpp
@@ -33,7 +33,7 @@ void SScoreRating::Construct(const FArguments& InArgs)
 	ScoreAttribute.Assign(*this, InArgs._Score);
 	ComponentIconAttribute.Assign(*this, InArgs._ComponentIcon);
 
-	for (int ScoreIndex = 0; ScoreIndex < MaxScore; ++ScoreIndex)
+	for (int32 ScoreIndex = 0; ScoreIndex < MaxScore; ++ScoreIndex)
 	{
 		LayoutBox->AddSlot()
 			.AutoWidth()
@@ -48,11 +48,11 @@ void SScoreRating::Construct(const FArguments& InArgs)
 						})
 					.OnCheckStateChanged(FOnBinaryCheckStateChanged::CreateLambda([this, ScoreIndex](bool bIsChecked)
 						{
-							const int NewScore = ScoreIndex + 1;
+							const int32 NewScore = ScoreIndex + 1;
 							ScoreAttribute.Set(*this, NewScore != ScoreAttribute.Get() ? NewScore : 0);
 						}))
 					.OnHoverStarted_Raw(this, &SScoreRating::SetHoverIndex, ScoreIndex)
-					.OnHoverFinished_Raw(this, &SScoreRating::SetHoverIndex, static_cast<int>(INDEX_NONE))
+					.OnHoverFinished_Raw(this, &SScoreRating::SetHoverIndex, static_cast<int32>(INDEX_NONE))
 			];
 	}
 
diff --git a/Engine/Plugins/MovieScene/LevelSequenceEditor/Source/LevelSequenceEditor/Private/Widgets/SScoreRating.h b/Engine/Plugins/MovieScene/LevelSequenceEditor/Source/LevelSequenceEditor/Private/Widgets/SScoreRating.h
--- a/Engine/Plugins/MovieScene/LevelSequenceEditor/Source/LevelSequenceEditor/Private/Widgets/SScoreRating.h
+++ b/Engine/Plugins/MovieScene/LevelSequenceEditor/Source/LevelSequenceEditor/Private/Widgets/SScoreRating.h
@@ -8,7 +8,7 @@
 #include "Widgets/SBoxPanel.h"
 
 /** Delegate that is executed when the score is changed */
-DECLARE_DELEGATE_OneParam(FOnScoreChanged, int )
+DECLARE_DELEGATE_OneParam(FOnScoreChanged, int32 )
 
 /** Widget for viewing and setting a score rating. */
 class SScoreRating : public SCompoundWidget
@@ -24,7 +24,7 @@ public:
 		{}
 
 		/** The max score shown by the widget */
-		SLATE_ARGUMENT( int, MaxScore )
+		SLATE_ARGUMENT( int32, MaxScore )
 
 		/** The icon brush to use for each score rating component */
 		SLATE_ATTRIBUTE( const FSlateBrush*, ComponentIcon )
@@ -33,7 +33,7 @@ public:
 		SLATE_EVENT( FOnScoreChanged, OnScoreChanged)
 
 		/** The current score */
-		SLATE_ATTRIBUTE( int, Score )
+		SLATE_ATTRIBUTE( int32, Score )
 	SLATE_END_ARGS()
 
 	SScoreRating();
@@ -42,7 +42,7 @@ public:
 
 private:
 	/** Set the current hover index */
-	inline void SetHoverIndex(int Index) { HoverIndex = Index; }
+	inline void SetHoverIndex(int32 Index) { HoverIndex = Index; }
 
 	/** Notify the score has changed with the new score */
 	void NotifyScoreChanged();
@@ -51,17 +51,17 @@ private:
 	TSharedPtr<SHorizontalBox> LayoutBox;
 
 	/** Max possible score to display */
-	int MaxScore = 3;
+	int32 MaxScore = 3;
 
 	/** Event triggered when the score is changed. */
 	FOnScoreChanged OnScoreChanged;
 
 	/** Current score rating index that the mouse is hovering over */
-	int HoverIndex = 0;
+	int32 HoverIndex = 0;
 
 	/** The slate brush to use to draw the component icons. */
 	TSlateAttribute<const FSlateBrush*> ComponentIconAttribute;
 
 	/** The current score */
-	TSlateAttribute<int> ScoreAttribute;
+	TSlateAttribute<int32> ScoreAttribute;
 };
