ShotManagement: Fixed issue where CineAssembly Edit Properties window would not keep sync with Favorite Rating if a rating was set in the Edit Property window.

SlateAttribute no longer being set on state change, event is notified directly.

#rnx

diff --git a/Engine/Plugins/MovieScene/LevelSequenceEditor/Source/LevelSequenceEditor/Private/Widgets/SScoreRating.cpp b/Engine/Plugins/MovieScene/LevelSequenceEditor/Source/LevelSequenceEditor/Private/Widgets/SScoreRating.cpp
--- a/Engine/Plugins/MovieScene/LevelSequenceEditor/Source/LevelSequenceEditor/Private/Widgets/SScoreRating.cpp
+++ b/Engine/Plugins/MovieScene/LevelSequenceEditor/Source/LevelSequenceEditor/Private/Widgets/SScoreRating.cpp
@@ -9,12 +9,7 @@ SLATE_IMPLEMENT_WIDGET(SScoreRating)
 void SScoreRating::PrivateRegisterAttributes(FSlateAttributeInitializer& AttributeInitializer)
 {
 	SLATE_ADD_MEMBER_ATTRIBUTE_DEFINITION_WITH_NAME(AttributeInitializer, "ComponentIconAttribute", ComponentIconAttribute, EInvalidateWidgetReason::Layout);
-	SLATE_ADD_MEMBER_ATTRIBUTE_DEFINITION_WITH_NAME(AttributeInitializer, "ScoreAttribute", ScoreAttribute, EInvalidateWidgetReason::Layout)
-	.OnValueChanged(FSlateAttributeDescriptor::FAttributeValueChangedDelegate::CreateLambda([](SWidget& Widget)
-		{
-			static_cast<SScoreRating&>(Widget).NotifyScoreChanged();
-		}
-	));
+	SLATE_ADD_MEMBER_ATTRIBUTE_DEFINITION_WITH_NAME(AttributeInitializer, "ScoreAttribute", ScoreAttribute, EInvalidateWidgetReason::Layout);
 }
 
 SScoreRating::SScoreRating() :
@@ -49,7 +44,7 @@ void SScoreRating::Construct(const FArguments& InArgs)
 					.OnCheckStateChanged(FOnBinaryCheckStateChanged::CreateLambda([this, ScoreIndex](bool bIsChecked)
 						{
 							const int32 NewScore = ScoreIndex + 1;
-							ScoreAttribute.Set(*this, NewScore != ScoreAttribute.Get() ? NewScore : 0);
+							NotifyScoreChanged(NewScore != ScoreAttribute.Get() ? NewScore : 0);
 						}))
 					.OnHoverStarted_Raw(this, &SScoreRating::SetHoverIndex, ScoreIndex)
 					.OnHoverFinished_Raw(this, &SScoreRating::SetHoverIndex, static_cast<int32>(INDEX_NONE))
@@ -62,9 +57,9 @@ void SScoreRating::Construct(const FArguments& InArgs)
 		];
 }
 
-void SScoreRating::NotifyScoreChanged()
+void SScoreRating::NotifyScoreChanged(int32 NewScore)
 {
-	OnScoreChanged.ExecuteIfBound(ScoreAttribute.Get());
+	OnScoreChanged.ExecuteIfBound(NewScore);
 }
 
 #undef LOCTEXT_NAMESPACE
diff --git a/Engine/Plugins/MovieScene/LevelSequenceEditor/Source/LevelSequenceEditor/Private/Widgets/SScoreRating.h b/Engine/Plugins/MovieScene/LevelSequenceEditor/Source/LevelSequenceEditor/Private/Widgets/SScoreRating.h
--- a/Engine/Plugins/MovieScene/LevelSequenceEditor/Source/LevelSequenceEditor/Private/Widgets/SScoreRating.h
+++ b/Engine/Plugins/MovieScene/LevelSequenceEditor/Source/LevelSequenceEditor/Private/Widgets/SScoreRating.h
@@ -45,7 +45,7 @@ private:
 	inline void SetHoverIndex(int32 Index) { HoverIndex = Index; }
 
 	/** Notify the score has changed with the new score */
-	void NotifyScoreChanged();
+	void NotifyScoreChanged(int32 NewScore);
 
 	/** The overall layout widget */
 	TSharedPtr<SHorizontalBox> LayoutBox;
