ShotManagement: Rename SScoreRatingComponent to SScoreRatingElement

#rnx

diff --git a/Engine/Plugins/MovieScene/LevelSequenceEditor/Source/LevelSequenceEditor/Private/Styles/LevelSequenceEditorStyle.h b/Engine/Plugins/MovieScene/LevelSequenceEditor/Source/LevelSequenceEditor/Private/Styles/LevelSequenceEditorStyle.h
--- a/Engine/Plugins/MovieScene/LevelSequenceEditor/Source/LevelSequenceEditor/Private/Styles/LevelSequenceEditorStyle.h
+++ b/Engine/Plugins/MovieScene/LevelSequenceEditor/Source/LevelSequenceEditor/Private/Styles/LevelSequenceEditorStyle.h
@@ -106,7 +106,7 @@ public:
 
 		Set("FlaggedWidget", FlaggedBoxStyle);
 
-		FCheckBoxStyle ScoreRatingComponentStyle = FCheckBoxStyle()
+		FCheckBoxStyle ScoreRatingElementStyle = FCheckBoxStyle()
 			.SetCheckBoxType(ESlateCheckBoxType::ToggleButton)
 			.SetForegroundColor(FSlateColor(FLinearColor::White))
 			.SetCheckedForegroundColor(FStyleColors::AccentBlue)
@@ -117,7 +117,7 @@ public:
 			.SetBorderBackgroundColor(FSlateColor(FLinearColor::Transparent))
 			.SetPadding(FMargin(8, 3));;
 
-		Set("ScoreRatingComponent", ScoreRatingComponentStyle);
+		Set("ScoreRatingElement", ScoreRatingElementStyle);
 
 		FSlateStyleRegistry::RegisterSlateStyle(*this);
 	 }
diff --git a/Engine/Plugins/MovieScene/LevelSequenceEditor/Source/LevelSequenceEditor/Private/Widgets/SScoreRating.cpp b/Engine/Plugins/MovieScene/LevelSequenceEditor/Source/LevelSequenceEditor/Private/Widgets/SScoreRating.cpp
--- a/Engine/Plugins/MovieScene/LevelSequenceEditor/Source/LevelSequenceEditor/Private/Widgets/SScoreRating.cpp
+++ b/Engine/Plugins/MovieScene/LevelSequenceEditor/Source/LevelSequenceEditor/Private/Widgets/SScoreRating.cpp
@@ -1,19 +1,19 @@
 // Copyright Epic Games, Inc. All Rights Reserved.
 
 #include "Widgets/SScoreRating.h"
-#include "Widgets/SScoreRatingComponent.h"
+#include "Widgets/SScoreRatingElement.h"
 
 #define LOCTEXT_NAMESPACE "SScoreRating"
 
 SLATE_IMPLEMENT_WIDGET(SScoreRating)
 void SScoreRating::PrivateRegisterAttributes(FSlateAttributeInitializer& AttributeInitializer)
 {
-	SLATE_ADD_MEMBER_ATTRIBUTE_DEFINITION_WITH_NAME(AttributeInitializer, "ComponentIconAttribute", ComponentIconAttribute, EInvalidateWidgetReason::Layout);
+	SLATE_ADD_MEMBER_ATTRIBUTE_DEFINITION_WITH_NAME(AttributeInitializer, "ElementIconAttribute", ElementIconAttribute, EInvalidateWidgetReason::Layout);
 	SLATE_ADD_MEMBER_ATTRIBUTE_DEFINITION_WITH_NAME(AttributeInitializer, "ScoreAttribute", ScoreAttribute, EInvalidateWidgetReason::Layout);
 }
 
 SScoreRating::SScoreRating() :
-	ComponentIconAttribute(*this),
+	ElementIconAttribute(*this),
 	ScoreAttribute(*this)
 {
 
@@ -26,7 +26,7 @@ void SScoreRating::Construct(const FArguments& InArgs)
 	MaxScore = FMath::Clamp(InArgs._MaxScore, 2, 10);
 	OnScoreChanged = InArgs._OnScoreChanged;
 	ScoreAttribute.Assign(*this, InArgs._Score);
-	ComponentIconAttribute.Assign(*this, InArgs._ComponentIcon);
+	ElementIconAttribute.Assign(*this, InArgs._ElementIcon);
 
 	for (int32 ScoreIndex = 0; ScoreIndex < MaxScore; ++ScoreIndex)
 	{
@@ -34,8 +34,8 @@ void SScoreRating::Construct(const FArguments& InArgs)
 			.AutoWidth()
 			.Padding(FMargin(0.0, 0.0, 4.0, 0.0))
 			[
-				SNew(SScoreRatingComponent)
-					.Icon(InArgs._ComponentIcon)
+				SNew(SScoreRatingElement)
+					.Icon(ElementIconAttribute.Get())
 					.ToolTipText(FText::Format(LOCTEXT("SetScoreRatingTooltip", "Score as {0}"), ScoreIndex + 1))
 					.IsChecked_Lambda([&Score = this->ScoreAttribute, &HoverIndex = this->HoverIndex, ScoreIndex]() -> bool
 						{
diff --git a/Engine/Plugins/MovieScene/LevelSequenceEditor/Source/LevelSequenceEditor/Private/Widgets/SScoreRating.h b/Engine/Plugins/MovieScene/LevelSequenceEditor/Source/LevelSequenceEditor/Private/Widgets/SScoreRating.h
--- a/Engine/Plugins/MovieScene/LevelSequenceEditor/Source/LevelSequenceEditor/Private/Widgets/SScoreRating.h
+++ b/Engine/Plugins/MovieScene/LevelSequenceEditor/Source/LevelSequenceEditor/Private/Widgets/SScoreRating.h
@@ -18,7 +18,7 @@ class SScoreRating : public SCompoundWidget
 public:
 	SLATE_BEGIN_ARGS(SScoreRating)
 		: _MaxScore(3)
-		, _ComponentIcon(FAppStyle::GetBrush("Icons.Star"))
+		, _ElementIcon(FAppStyle::GetBrush("Icons.Star"))
 		, _OnScoreChanged()
 		, _Score(0)
 		{}
@@ -26,8 +26,8 @@ public:
 		/** The max score shown by the widget */
 		SLATE_ARGUMENT( int32, MaxScore )
 
-		/** The icon brush to use for each score rating component */
-		SLATE_ATTRIBUTE( const FSlateBrush*, ComponentIcon )
+		/** The icon brush to use for each score rating element */
+		SLATE_ATTRIBUTE( const FSlateBrush*, ElementIcon )
 
 		/** Event triggered when the score is changed. */
 		SLATE_EVENT( FOnScoreChanged, OnScoreChanged)
@@ -59,8 +59,8 @@ private:
 	/** Current score rating index that the mouse is hovering over */
 	int32 HoverIndex = 0;
 
-	/** The slate brush to use to draw the component icons. */
-	TSlateAttribute<const FSlateBrush*> ComponentIconAttribute;
+	/** The slate brush to use to draw the element icons. */
+	TSlateAttribute<const FSlateBrush*> ElementIconAttribute;
 
 	/** The current score */
 	TSlateAttribute<int32> ScoreAttribute;
diff --git a/Engine/Plugins/MovieScene/LevelSequenceEditor/Source/LevelSequenceEditor/Private/Widgets/SScoreRatingComponent.cpp b/Engine/Plugins/MovieScene/LevelSequenceEditor/Source/LevelSequenceEditor/Private/Widgets/SScoreRatingComponent.cpp
deleted file mode 100644
diff --git a/Engine/Plugins/MovieScene/LevelSequenceEditor/Source/LevelSequenceEditor/Private/Widgets/SScoreRatingComponent.h b/Engine/Plugins/MovieScene/LevelSequenceEditor/Source/LevelSequenceEditor/Private/Widgets/SScoreRatingComponent.h
deleted file mode 100644
diff --git a/Engine/Plugins/MovieScene/LevelSequenceEditor/Source/LevelSequenceEditor/Private/Widgets/SScoreRatingElement.cpp b/Engine/Plugins/MovieScene/LevelSequenceEditor/Source/LevelSequenceEditor/Private/Widgets/SScoreRatingElement.cpp
new file mode 100644
--- /dev/null
+++ b/Engine/Plugins/MovieScene/LevelSequenceEditor/Source/LevelSequenceEditor/Private/Widgets/SScoreRatingElement.cpp
@@ -0,0 +1,62 @@
+// Copyright Epic Games, Inc. All Rights Reserved.
+
+#include "Widgets/SScoreRatingElement.h"
+#include "Styles/LevelSequenceEditorStyle.h"
+#include "Widgets/Images/SImage.h"
+#include "Widgets/Input/SCheckBox.h"
+
+#define LOCTEXT_NAMESPACE "SScoreRatingElement"
+
+SLATE_IMPLEMENT_WIDGET(SScoreRatingElement)
+void SScoreRatingElement::PrivateRegisterAttributes(FSlateAttributeInitializer& AttributeInitializer)
+{
+	SLATE_ADD_MEMBER_ATTRIBUTE_DEFINITION_WITH_NAME(AttributeInitializer, "IconAttribute", IconAttribute, EInvalidateWidgetReason::Layout);
+	SLATE_ADD_MEMBER_ATTRIBUTE_DEFINITION_WITH_NAME(AttributeInitializer, "IsCheckedAttribute", IsCheckedAttribute, EInvalidateWidgetReason::Layout);
+}
+
+SScoreRatingElement::SScoreRatingElement() :
+	IconAttribute(*this),
+	IsCheckedAttribute(*this)
+{
+
+}
+
+void SScoreRatingElement::Construct(const FArguments& InArgs)
+{
+	OnHoverStarted = InArgs._OnHoverStarted;
+	OnHoverFinished = InArgs._OnHoverFinished;
+	OnCheckStateChanged = InArgs._OnCheckStateChanged;
+	IconAttribute.Assign(*this, InArgs._Icon);
+	IsCheckedAttribute.Assign(*this, InArgs._IsChecked);
+
+	ChildSlot
+		[
+			SNew(SCheckBox)
+				.Padding(FMargin(0.0f))
+				.Style(&FLevelSequenceEditorStyle::Get()->GetWidgetStyle<FCheckBoxStyle>("ScoreRatingElement"))
+				.IsChecked_Lambda([&IsChecked = this->IsCheckedAttribute]() -> ECheckBoxState { return IsChecked.Get() ? ECheckBoxState::Checked : ECheckBoxState::Unchecked; })
+				.OnCheckStateChanged(FOnCheckStateChanged::CreateLambda([&OnCheckStateChangedRef = this->OnCheckStateChanged](ECheckBoxState State) 
+					{
+						OnCheckStateChangedRef.ExecuteIfBound(State == ECheckBoxState::Checked);
+					}))
+				[
+					SNew(SImage)
+						.ColorAndOpacity(FSlateColor::UseForeground())
+						.Image(IconAttribute.Get())
+				]
+		];
+}
+
+void SScoreRatingElement::OnMouseEnter(const FGeometry& MyGeometry, const FPointerEvent& MouseEvent)
+{
+	SCompoundWidget::OnMouseEnter(MyGeometry, MouseEvent);
+	OnHoverStarted.ExecuteIfBound();
+}
+
+void SScoreRatingElement::OnMouseLeave(const FPointerEvent& MouseEvent)
+{
+	SCompoundWidget::OnMouseLeave(MouseEvent);
+	OnHoverFinished.ExecuteIfBound();
+}
+
+#undef LOCTEXT_NAMESPACE
diff --git a/Engine/Plugins/MovieScene/LevelSequenceEditor/Source/LevelSequenceEditor/Private/Widgets/SScoreRatingElement.h b/Engine/Plugins/MovieScene/LevelSequenceEditor/Source/LevelSequenceEditor/Private/Widgets/SScoreRatingElement.h
new file mode 100644
--- /dev/null
+++ b/Engine/Plugins/MovieScene/LevelSequenceEditor/Source/LevelSequenceEditor/Private/Widgets/SScoreRatingElement.h
@@ -0,0 +1,74 @@
+// Copyright Epic Games, Inc. All Rights Reserved.
+
+#pragma once
+
+#include "Styling/AppStyle.h"
+#include "Widgets/SCompoundWidget.h"
+#include "Widgets/DeclarativeSyntaxSupport.h"
+
+/** Delegate that is executed when a Score Rating Element hover event happens */
+DECLARE_DELEGATE(FOnScoreRatingElementHoverEvent);
+
+/** Delegate that is executed when a binary check state change occurs, either to/from true/false */
+DECLARE_DELEGATE_OneParam(FOnBinaryCheckStateChanged, bool);
+
+/** 
+ * Specialized widget for a score value to expose when mouse hover events occur.
+ * Ideally this would be contained entirely in SScoreRating however the need to get the hover events
+ * mean we require an extra wrapper class.
+ */
+class SScoreRatingElement : public SCompoundWidget
+{
+	SLATE_DECLARE_WIDGET(SScoreRatingElement, SCompoundWidget)
+
+public:
+	SLATE_BEGIN_ARGS(SScoreRatingElement)
+		: _OnHoverStarted()
+		, _OnHoverFinished()
+		, _OnCheckStateChanged()
+		, _Icon(FAppStyle::GetBrush("Icons.Star"))
+		, _IsChecked(false)
+		{
+		}
+
+		/** Called when a hover event starts */
+		SLATE_EVENT( FOnScoreRatingElementHoverEvent, OnHoverStarted )
+
+		/** Called when a hover event finishes */
+		SLATE_EVENT( FOnScoreRatingElementHoverEvent, OnHoverFinished )
+		
+		/** Called when a a state change occurs on the Rating Star through user input */
+		SLATE_EVENT( FOnBinaryCheckStateChanged, OnCheckStateChanged )
+
+		/** The Slate brush to use for this element */
+		SLATE_ATTRIBUTE( const FSlateBrush*, Icon )
+
+		/** Whether the Score Rating Element is currently in its checked state. */
+		SLATE_ATTRIBUTE( bool, IsChecked )
+
+	SLATE_END_ARGS()
+
+	SScoreRatingElement();
+
+	void Construct(const FArguments& InArgs);
+
+	// SWidget Interface
+	void OnMouseEnter(const FGeometry& MyGeometry, const FPointerEvent& MouseEvent) override;
+	void OnMouseLeave(const FPointerEvent& MouseEvent) override;
+
+private:
+	/** Delegate fired when a hover event starts */
+	FOnScoreRatingElementHoverEvent OnHoverStarted;
+
+	/** Delegate fired when a hover event finishes */
+	FOnScoreRatingElementHoverEvent OnHoverFinished;
+
+	/** Delegate fired when a state change event occurs */
+	FOnBinaryCheckStateChanged OnCheckStateChanged;
+
+	/** The slate brush to use to draw the icon. */
+	TSlateAttribute<const FSlateBrush*> IconAttribute;
+
+	/** Are we checked */
+	TSlateAttribute<bool> IsCheckedAttribute;
+};
