Sequencer: Disallow picking a proxy binding that's already bound to the current one

diff --git a/Engine/Plugins/MovieScene/LevelSequenceEditor/Source/LevelSequenceEditor/Private/LevelSequenceEditorSubsystem.cpp b/Engine/Plugins/MovieScene/LevelSequenceEditor/Source/LevelSequenceEditor/Private/LevelSequenceEditorSubsystem.cpp
--- a/Engine/Plugins/MovieScene/LevelSequenceEditor/Source/LevelSequenceEditor/Private/LevelSequenceEditorSubsystem.cpp
+++ b/Engine/Plugins/MovieScene/LevelSequenceEditor/Source/LevelSequenceEditor/Private/LevelSequenceEditorSubsystem.cpp
@@ -1,6 +1,7 @@
 // Copyright Epic Games, Inc. All Rights Reserved.
 
 #include "LevelSequenceEditorSubsystem.h"
+#include "Compilation/MovieSceneCompiledDataManager.h"
 #include "MVVM/ViewModels/SequencerEditorViewModel.h"
 #include "Scripting/SequencerModuleScriptingLayer.h"
 #include "SequencerCurveEditorObject.h"
@@ -105,6 +106,51 @@ namespace UE::Sequencer
 			, PreOnPicked(InPreOnPicked)
 		{
 			Initialize();
+			
+			InitializeExistingBindingIDs();
+		}
+
+		void GatherExistingBindingIDs(UMovieSceneSequence* InSequence, FMovieSceneSequenceID InSequenceID)
+		{
+			if (UMovieScene* MovieScene = InSequence->GetMovieScene())
+			{
+				for (int32 Index = 0; Index < MovieScene->GetPossessableCount(); ++Index)
+				{
+					FMovieScenePossessable& Possessable = MovieScene->GetPossessable(Index);
+
+					if (Possessable.GetSpawnableObjectBindingID().GetGuid() == ObjectBindingID)
+					{
+						UE::MovieScene::FFixedObjectBindingID BindingID(Possessable.GetGuid(), InSequenceID);
+						ExistingBindingIDs.Add(BindingID);
+					}
+				}
+			}
+		}
+
+		void InitializeExistingBindingIDs()
+		{
+			TSharedPtr<ISequencer> Sequencer = WeakSequencer.Pin();
+			if (!Sequencer)
+			{
+				return;
+			}
+
+			const FMovieSceneRootEvaluationTemplateInstance& RootInstance = Sequencer->GetEvaluationTemplate();
+			const FMovieSceneSequenceHierarchy* Hierarchy = RootInstance.GetCompiledDataManager()->FindHierarchy(RootInstance.GetCompiledDataID());
+			if (!Hierarchy)
+			{
+				return;
+			}
+
+			GatherExistingBindingIDs(Sequencer->GetRootMovieSceneSequence(), MovieSceneSequenceID::Root);
+
+			for (const TTuple<FMovieSceneSequenceID, FMovieSceneSubSequenceData>& Pair : Hierarchy->AllSubSequenceData())
+			{
+				if (UMovieSceneSequence* Sequence = Pair.Value.GetSequence())
+				{
+					GatherExistingBindingIDs(Sequence, Pair.Key);
+				}
+			}
 		}
 
 		virtual UMovieSceneSequence* GetSequence() const override
@@ -152,12 +198,19 @@ namespace UE::Sequencer
 				return false;
 			}
 
+			// Disallow picking a binding that already is a proxy binding to this node
+			if (ExistingBindingIDs.Contains(Node->BindingID))
+			{
+				return false;
+			}
+
 			return FMovieSceneObjectBindingIDPicker::IsNodeAllowed(Node);
 		}
 
 	private:
 		FGuid ObjectBindingID;
 		TFunction<void(FMovieSceneObjectBindingID)> PreOnPicked;
+		TArray<UE::MovieScene::FFixedObjectBindingID> ExistingBindingIDs;
 	};
 }
 
