ShotManagement: Wrapping ShotMetaData Favorite Rating modifications with Transactions to support undo/redo.

#rnx

diff --git a/Engine/Plugins/MovieScene/LevelSequenceEditor/Source/LevelSequenceEditor/Private/Widgets/SLevelSequenceFavoriteRating.cpp b/Engine/Plugins/MovieScene/LevelSequenceEditor/Source/LevelSequenceEditor/Private/Widgets/SLevelSequenceFavoriteRating.cpp
--- a/Engine/Plugins/MovieScene/LevelSequenceEditor/Source/LevelSequenceEditor/Private/Widgets/SLevelSequenceFavoriteRating.cpp
+++ b/Engine/Plugins/MovieScene/LevelSequenceEditor/Source/LevelSequenceEditor/Private/Widgets/SLevelSequenceFavoriteRating.cpp
@@ -2,6 +2,7 @@
 
 #include "Widgets/SLevelSequenceFavoriteRating.h"
 #include "MetaData/MovieSceneShotMetaData.h"
+#include "ScopedTransaction.h"
 #include "Styles/LevelSequenceEditorStyle.h"
 #include "Styling/StyleColors.h"
 #include "Widgets/SScoreRating.h"
@@ -39,7 +40,15 @@ void SLevelSequenceFavoriteRating::Construct(const FArguments& InArgs)
 								{
 									if (UMovieSceneShotMetaData* MetaData = WeakLevelSequence.IsValid() ? WeakLevelSequence.Pin()->FindOrAddMetaData<UMovieSceneShotMetaData>() : nullptr)
 									{
-										MetaData->SetIsNoGood(State == ECheckBoxState::Checked);
+										const bool bIsNoGood = State == ECheckBoxState::Checked;
+										const FText TransactionSessionName = bIsNoGood
+											? LOCTEXT("SetLevelSequenceIsNoGoodTransaction", "Set Level Sequence Is No Good")
+											: LOCTEXT("UnsetLevelSequenceIsNoGoodTransaction", "Unset Level Sequence Is No Good");
+
+										const FScopedTransaction Transaction(TransactionSessionName);
+										MetaData->SetFlags(RF_Transactional);
+										MetaData->Modify();
+										MetaData->SetIsNoGood(bIsNoGood);
 									}
 								}))
 							[
@@ -70,7 +79,15 @@ void SLevelSequenceFavoriteRating::Construct(const FArguments& InArgs)
 								{
 									if (UMovieSceneShotMetaData* MetaData = WeakLevelSequence.IsValid() ? WeakLevelSequence.Pin()->FindOrAddMetaData<UMovieSceneShotMetaData>() : nullptr)
 									{
-										MetaData->SetIsFlagged(State == ECheckBoxState::Checked);
+										const bool bIsFlagged = State == ECheckBoxState::Checked;
+										const FText TransactionSessionName = bIsFlagged 
+											? LOCTEXT("SetLevelSequenceIsFlaggedTransaction", "Set Level Sequence Is Flagged")
+											: LOCTEXT("UnsetLevelSequenceIsFlaggedTransaction", "Unset Level Sequence Is Flagged");
+
+										const FScopedTransaction Transaction(TransactionSessionName);
+										MetaData->SetFlags(RF_Transactional);
+										MetaData->Modify();
+										MetaData->SetIsFlagged(bIsFlagged);
 									}
 								}))
 							[
@@ -91,6 +108,9 @@ void SLevelSequenceFavoriteRating::Construct(const FArguments& InArgs)
 								{
 									if (UMovieSceneShotMetaData* MetaData = WeakLevelSequence.IsValid() ? WeakLevelSequence.Pin()->FindOrAddMetaData<UMovieSceneShotMetaData>() : nullptr)
 									{
+										const FScopedTransaction Transaction(LOCTEXT("SetLevelSequenceFavoriteRatingTransaction", "Set Level Sequence Favorite Rating"));
+										MetaData->SetFlags(RF_Transactional);
+										MetaData->Modify();
 										if (NewScore > 0)
 										{
 											MetaData->SetFavoriteRating(NewScore);
