Sequencer: Added nullptr checks for GetEventContexts

#rnx

diff --git a/Engine/Plugins/MovieScene/LevelSequenceEditor/Source/LevelSequenceEditor/Private/Misc/LevelSequencePlaybackContext.cpp b/Engine/Plugins/MovieScene/LevelSequenceEditor/Source/LevelSequenceEditor/Private/Misc/LevelSequencePlaybackContext.cpp
--- a/Engine/Plugins/MovieScene/LevelSequenceEditor/Source/LevelSequenceEditor/Private/Misc/LevelSequencePlaybackContext.cpp
+++ b/Engine/Plugins/MovieScene/LevelSequenceEditor/Source/LevelSequenceEditor/Private/Misc/LevelSequencePlaybackContext.cpp
@@ -242,8 +242,13 @@ IMovieScenePlaybackClient* FLevelSequencePlaybackContext::GetPlaybackClientAsInt
 TArray<UObject*> FLevelSequencePlaybackContext::GetEventContexts() const
 {
 	TArray<UObject*> Contexts;
-	UWorld* ContextWorld = GetPlaybackContext()->GetWorld();
-	ULevelSequencePlayer::GetEventContexts(*ContextWorld, Contexts);
+	if (UObject* PlaybackContext = GetPlaybackContext())
+	{
+		if (UWorld* ContextWorld = PlaybackContext->GetWorld())
+		{
+			ULevelSequencePlayer::GetEventContexts(*ContextWorld, Contexts);
+		}
+	}
 	return Contexts;
 }
 
