Sequencer: Move BakeTransform to SequencerUtilities so that level sequences can access it

Bake Transform is now accessible in Camera Animation Sequences and Template Sequences

Also, fixed a minor issue where the baking settings wouldn't persist

diff --git a/Engine/Plugins/MovieScene/LevelSequenceEditor/Source/LevelSequenceEditor/Private/LevelSequenceEditorCommands.cpp b/Engine/Plugins/MovieScene/LevelSequenceEditor/Source/LevelSequenceEditor/Private/LevelSequenceEditorCommands.cpp
--- a/Engine/Plugins/MovieScene/LevelSequenceEditor/Source/LevelSequenceEditor/Private/LevelSequenceEditorCommands.cpp
+++ b/Engine/Plugins/MovieScene/LevelSequenceEditor/Source/LevelSequenceEditor/Private/LevelSequenceEditorCommands.cpp
@@ -32,7 +32,6 @@ void FLevelSequenceEditorCommands::RegisterCommands()
 	UI_COMMAND(SnapSectionsToTimelineUsingSourceTimecode, "Snap Sections to Timeline using Source Timecode", "Snap selected sections to the position in the timeline matching their source timecode", EUserInterfaceActionType::Button, FInputChord());
 	UI_COMMAND(SyncSectionsUsingSourceTimecode, "Sync Sections using Source Timecode", "Synchronize sections to the first selected section using source timecode", EUserInterfaceActionType::Button, FInputChord());
 
-	UI_COMMAND(BakeTransform, "Bake Transform", "Bake transform in world space, removing any existing transform and attach tracks", EUserInterfaceActionType::Button, FInputChord());
 	UI_COMMAND(FixActorReferences, "Fix Actor References", "Try to automatically fix up broken actor bindings", EUserInterfaceActionType::Button, FInputChord());
 
 	UI_COMMAND(AddActorsToBinding, "Add Selected", "Add selected actors to this track", EUserInterfaceActionType::Button, FInputChord());
diff --git a/Engine/Plugins/MovieScene/LevelSequenceEditor/Source/LevelSequenceEditor/Private/LevelSequenceEditorCommands.h b/Engine/Plugins/MovieScene/LevelSequenceEditor/Source/LevelSequenceEditor/Private/LevelSequenceEditorCommands.h
--- a/Engine/Plugins/MovieScene/LevelSequenceEditor/Source/LevelSequenceEditor/Private/LevelSequenceEditorCommands.h
+++ b/Engine/Plugins/MovieScene/LevelSequenceEditor/Source/LevelSequenceEditor/Private/LevelSequenceEditorCommands.h
@@ -22,7 +22,6 @@ public:
 	TSharedPtr<FUICommandInfo> SnapSectionsToTimelineUsingSourceTimecode;
 	TSharedPtr<FUICommandInfo> SyncSectionsUsingSourceTimecode;
 
-	TSharedPtr<FUICommandInfo> BakeTransform;
 	TSharedPtr<FUICommandInfo> FixActorReferences;
 
 	TSharedPtr<FUICommandInfo> AddActorsToBinding;
diff --git a/Engine/Plugins/MovieScene/LevelSequenceEditor/Source/LevelSequenceEditor/Private/LevelSequenceEditorSubsystem.cpp b/Engine/Plugins/MovieScene/LevelSequenceEditor/Source/LevelSequenceEditor/Private/LevelSequenceEditorSubsystem.cpp
--- a/Engine/Plugins/MovieScene/LevelSequenceEditor/Source/LevelSequenceEditor/Private/LevelSequenceEditorSubsystem.cpp
+++ b/Engine/Plugins/MovieScene/LevelSequenceEditor/Source/LevelSequenceEditor/Private/LevelSequenceEditorSubsystem.cpp
@@ -17,19 +17,13 @@
 #include "MovieScene.h"
 #include "MovieSceneSpawnable.h"
 #include "SequencerUtilities.h"
-#include "Sections/MovieScene3DConstraintSection.h"
 #include "Selection.h"
-#include "Tracks/MovieScene3DTransformTrack.h"
 #include "Toolkits/AssetEditorToolkit.h"
-#include "Tracks/MovieScene3DConstraintTrack.h"
-#include "Tracks/MovieSceneCameraShakeTrack.h"
 
 #include "ActorTreeItem.h"
-#include "Editor.h"
 #include "PropertyEditorModule.h"
 #include "Widgets/SWindow.h"
 #include "Widgets/Input/SButton.h"
-#include "Widgets/Layout/SSpacer.h"
 #include "Modules/ModuleManager.h"
 
 #include "Camera/CameraComponent.h"
@@ -45,15 +39,10 @@
 #include "ScopedTransaction.h"
 #include "Widgets/Notifications/SNotificationList.h"
 #include "Framework/MultiBox/MultiBoxBuilder.h"
-#include "BakingAnimationKeySettings.h"
-#include "IStructureDetailsView.h"
 #include "Widgets/SWidget.h"
 #include "Widgets/SCompoundWidget.h"
 #include "Widgets/Layout/SBox.h"
-#include "Widgets/Input/NumericTypeInterface.h"
-#include "FrameNumberDetailsCustomization.h"
 #include "MovieSceneToolHelpers.h"
-#include "ActorForWorldTransforms.h"
 #include "KeyParams.h"
 #include "Sections/MovieSceneBindingLifetimeSection.h"
 #include "Tracks/MovieSceneBindingLifetimeTrack.h"
@@ -63,7 +52,6 @@
 #include "Bindings/MovieSceneSpawnableActorBindingCustomization.h"
 #include "Bindings/MovieSceneCustomBinding.h"
 #include "DetailWidgetRow.h"
-#include "IPropertyTypeCustomization.h"
 #include "IDetailChildrenBuilder.h"
 #include "Widgets/Input/SComboBox.h"
 #include "DetailLayoutBuilder.h"
@@ -655,10 +643,6 @@ void ULevelSequenceEditorSubsystem::Initialize(FSubsystemCollectionBase& Collect
 		FCanExecuteAction::CreateLambda(AreMovieSceneSectionsSelected, 2)
 	);
 
-	CommandList->MapAction(FLevelSequenceEditorCommands::Get().BakeTransform,
-		FExecuteAction::CreateUObject(this, &ULevelSequenceEditorSubsystem::BakeTransformInternal)
-	);
-
 	CommandList->MapAction(FLevelSequenceEditorCommands::Get().FixActorReferences,
 		FExecuteAction::CreateUObject(this, &ULevelSequenceEditorSubsystem::FixActorReferences)
 	);
@@ -692,7 +676,6 @@ void ULevelSequenceEditorSubsystem::Initialize(FSubsystemCollectionBase& Collect
 		
 		MenuBuilder.AddMenuEntry(FLevelSequenceEditorCommands::Get().SnapSectionsToTimelineUsingSourceTimecode);
 		MenuBuilder.AddMenuEntry(FLevelSequenceEditorCommands::Get().SyncSectionsUsingSourceTimecode);
-		MenuBuilder.AddMenuEntry(FLevelSequenceEditorCommands::Get().BakeTransform);
 		}));
 
 	SequencerModule.GetActionsMenuExtensibilityManager()->AddExtender(TransformMenuExtender);
@@ -1640,685 +1623,26 @@ void ULevelSequenceEditorSubsystem::SyncSectionsUsingSourceTimecode(const TArray
 	}
 }
 
-//////////////////////////////////////////////////////////////
-/// SBakeTransformWidget
-///////////////////////////////////////////////////////////
-
-DECLARE_DELEGATE_RetVal_OneParam(FReply, SBakeTransformOnBake, FBakingAnimationKeySettings);
-
-/** Widget allowing baking controls from one space to another */
-class  SBakeTransformWidget : public SCompoundWidget
-{
-public:
-
-	SLATE_BEGIN_ARGS(SBakeTransformWidget)
-		: _Sequencer(nullptr)
-	{}
-		SLATE_ARGUMENT(ISequencer*, Sequencer)
-		SLATE_ARGUMENT(FBakingAnimationKeySettings, Settings)
-		SLATE_EVENT(SBakeTransformOnBake, OnBake)
-		SLATE_END_ARGS()
-
-	void Construct(const FArguments& InArgs);
-	virtual ~SBakeTransformWidget() override {}
-
-	FReply OpenDialog(bool bModal = true);
-	void CloseDialog();
-
-private:
-
-	//used for setting up the details
-	TSharedPtr<TStructOnScope<FBakingAnimationKeySettings>> Settings;
-
-	ISequencer* Sequencer;
-
-	TWeakPtr<SWindow> DialogWindow;
-	TSharedPtr<IStructureDetailsView> DetailsView;
-};
-
-void SBakeTransformWidget::Construct(const FArguments& InArgs)
-{
-	check(InArgs._Sequencer);
-	check(InArgs._OnBake.IsBound());
-
-	Settings = MakeShared<TStructOnScope<FBakingAnimationKeySettings>>();
-	Settings->InitializeAs<FBakingAnimationKeySettings>();
-	*Settings = InArgs._Settings;
-	//always setting space to be parent as default, since stored space may not be available.
-	Sequencer = InArgs._Sequencer;
-
-	FStructureDetailsViewArgs StructureViewArgs;
-	StructureViewArgs.bShowObjects = true;
-	StructureViewArgs.bShowAssets = true;
-	StructureViewArgs.bShowClasses = true;
-	StructureViewArgs.bShowInterfaces = true;
-
-	FDetailsViewArgs ViewArgs;
-	ViewArgs.bAllowSearch = false;
-	ViewArgs.bHideSelectionTip = false;
-	ViewArgs.bShowObjectLabel = false;
-
-	FPropertyEditorModule& PropertyEditor = FModuleManager::Get().LoadModuleChecked<FPropertyEditorModule>(TEXT("PropertyEditor"));
-
-	DetailsView = PropertyEditor.CreateStructureDetailView(ViewArgs, StructureViewArgs, TSharedPtr<FStructOnScope>());
-
-	DetailsView->GetDetailsView()->RegisterInstancedCustomPropertyTypeLayout("FrameNumber",
-		FOnGetPropertyTypeCustomizationInstance::CreateSP(Sequencer, &ISequencer::MakeFrameNumberDetailsCustomization));
-	DetailsView->SetStructureData(Settings);
-
-	ChildSlot
-		[
-			SNew(SBorder)
-			.Visibility(EVisibility::Visible)
-		[
-			SNew(SVerticalBox)
-
-			
-			+ SVerticalBox::Slot()
-				.AutoHeight()
-				.Padding(0.f, 0.f, 0.f, 0.f)
-				[
-					DetailsView->GetWidget().ToSharedRef()
-				]
-
-			+ SVerticalBox::Slot()
-			.AutoHeight()
-			.Padding(0.f, 16.f, 0.f, 16.f)
-			[
-				SNew(SHorizontalBox)
-
-				+ SHorizontalBox::Slot()
-			.FillWidth(1.f)
-			[
-				SNew(SSpacer)
-			]
-
-			+ SHorizontalBox::Slot()
-			.AutoWidth()
-			.HAlign(HAlign_Left)
-			.VAlign(VAlign_Center)
-			.Padding(8.f, 0.f, 0.f, 0.f)
-			[
-				SNew(SButton)
-				.HAlign(HAlign_Center)
-			.ContentPadding(FAppStyle::GetMargin("StandardDialog.ContentPadding"))
-			.Text(LOCTEXT("OK", "OK"))
-			.OnClicked_Lambda([this, InArgs]()
-				{
-					FReply Reply = InArgs._OnBake.Execute( *(Settings->Get()));
-					CloseDialog();
-					return Reply;
-
-				})
-			]
-
-		+ SHorizontalBox::Slot()
-			.AutoWidth()
-			.HAlign(HAlign_Left)
-			.VAlign(VAlign_Center)
-			.Padding(8.f, 0.f, 16.f, 0.f)
-			[
-				SNew(SButton)
-				.HAlign(HAlign_Center)
-				.ContentPadding(FAppStyle::GetMargin("StandardDialog.ContentPadding"))
-				.Text(LOCTEXT("Cancel", "Cancel"))
-				.OnClicked_Lambda([this]()
-				{
-					CloseDialog();
-					return FReply::Handled();
-				})
-			]
-		]
-	]
-	];
-}
-
-class SBakeTransformDialogWindow : public SWindow
-{
-};
-
-FReply SBakeTransformWidget::OpenDialog(bool bModal)
-{
-	check(!DialogWindow.IsValid());
-
-	const FVector2D CursorPos = FSlateApplication::Get().GetCursorPos();
-
-	TSharedRef<SBakeTransformDialogWindow> Window = SNew(SBakeTransformDialogWindow)
-		.Title(LOCTEXT("SBakeTransformWidgetTitle", "Bake Transforms"))
-		.CreateTitleBar(true)
-		.Type(EWindowType::Normal)
-		.SizingRule(ESizingRule::Autosized)
-		.ScreenPosition(CursorPos)
-		.FocusWhenFirstShown(true)
-		.ActivationPolicy(EWindowActivationPolicy::FirstShown)
-		[
-			AsShared()
-		];
-
-	Window->SetWidgetToFocusOnActivate(AsShared());
-
-	DialogWindow = Window;
-
-	Window->MoveWindowTo(CursorPos);
-
-	if (bModal)
-	{
-		GEditor->EditorAddModalWindow(Window);
-	}
-	else
-	{
-		FSlateApplication::Get().AddWindow(Window);
-	}
-
-	return FReply::Handled();
-}
-
-void SBakeTransformWidget::CloseDialog()
-{
-	if (DialogWindow.IsValid())
-	{
-		DialogWindow.Pin()->RequestDestroyWindow();
-		DialogWindow.Reset();
-	}
-}
-
-void ULevelSequenceEditorSubsystem::BakeTransformInternal()
-{
-	TSharedPtr<ISequencer> Sequencer = GetActiveSequencer();
-	if (Sequencer == nullptr)
-	{
-		return;
-	}
-
-	UMovieScene* FocusedMovieScene = Sequencer->GetFocusedMovieSceneSequence()->GetMovieScene();
-	if (!FocusedMovieScene)
-	{
-		return;
-	}
-
-	FFrameRate TickResolution = FocusedMovieScene->GetTickResolution();
-	FFrameRate DisplayRate = FocusedMovieScene->GetDisplayRate();
-
-	TArray<FGuid> ObjectBindings;
-	Sequencer->GetSelectedObjects(ObjectBindings);
-
-	TArray<FMovieSceneBindingProxy> BindingProxies;
-	for (FGuid Guid : ObjectBindings)
-	{
-		BindingProxies.Add(FMovieSceneBindingProxy(Guid, Sequencer->GetFocusedMovieSceneSequence()));
-	}
-
-	static FBakingAnimationKeySettings Settings; //reuse the settings except for the range
-	Settings.StartFrame = UE::MovieScene::DiscreteInclusiveLower(FocusedMovieScene->GetPlaybackRange());
-	Settings.EndFrame = UE::MovieScene::DiscreteExclusiveUpper(FocusedMovieScene->GetPlaybackRange());
-
-	TSharedRef<SBakeTransformWidget> BakeWidget =
-		SNew(SBakeTransformWidget)
-		.Settings(Settings)
-		.Sequencer(Sequencer.Get())
-		.OnBake_Lambda([this, &BindingProxies, Sequencer](FBakingAnimationKeySettings InSettings)
-			{
-				FMovieSceneScriptingParams Params;
-				Params.TimeUnit = EMovieSceneTimeUnit::TickResolution;
-				BakeTransformWithSettings(BindingProxies, InSettings, Params);
-				return FReply::Handled();
-			});
-
-	BakeWidget->OpenDialog(true);
-}
-
-void ULevelSequenceEditorSubsystem::CalculateFramesPerGuid(TSharedPtr<ISequencer>& Sequencer, const FBakingAnimationKeySettings& InSettings, TMap<FGuid, ULevelSequenceEditorSubsystem::FBakeData>& OutBakeDataMap,
-	TSortedMap<FFrameNumber, FFrameNumber>& OutFrameMap)
-{
-	OutFrameMap.Reset();
-	TArray<FFrameNumber> Frames;
-	//we get all frames since we need to get the Actor PER FRAME in order to handle spanwables
-	MovieSceneToolHelpers::CalculateFramesBetween(Sequencer->GetFocusedMovieSceneSequence()->GetMovieScene(), InSettings.StartFrame, InSettings.EndFrame, InSettings.FrameIncrement, Frames);
-	if (InSettings.BakingKeySettings == EBakingKeySettings::AllFrames)
-	{
-		for (FFrameNumber& Frame : Frames)
-		{
-			OutFrameMap.Add(Frame,Frame);
-		}
-		for (TPair<FGuid, FBakeData>& BakeData : OutBakeDataMap)
-		{
-			BakeData.Value.KeyTimes.Reset();
-			BakeData.Value.KeyTimes = OutFrameMap;
-		}
-	}
-	else
-	{
-		for (TPair<FGuid, FBakeData>& BakeData : OutBakeDataMap)
-		{
-			FActorForWorldTransforms ActorForWorldTransforms;
-			FGuid Guid = BakeData.Key;
-
-			for (TWeakObjectPtr<> RuntimeObject : Sequencer->FindObjectsInCurrentSequence(Guid))
-			{
-				UActorComponent* ActorComponent = nullptr;
-				AActor* Actor = Cast<AActor>(RuntimeObject.Get());
-				if (!Actor)
-				{
-					ActorComponent = Cast<UActorComponent>(RuntimeObject.Get());
-					if (ActorComponent)
-					{
-						Actor = ActorComponent->GetOwner();
-					}
-				}
-				ActorForWorldTransforms.Actor = Actor;
-				ActorForWorldTransforms.Component = Cast<USceneComponent>(ActorComponent);
-				BakeData.Value.KeyTimes.Reset();
-				MovieSceneToolHelpers::GetActorsAndParentsKeyFrames(Sequencer.Get(), ActorForWorldTransforms,
-					InSettings.StartFrame, InSettings.EndFrame, BakeData.Value.KeyTimes);
-				OutFrameMap.Append(BakeData.Value.KeyTimes);
-			}
-		}
-	}
-}
-
 bool ULevelSequenceEditorSubsystem::BakeTransformWithSettings(const TArray<FMovieSceneBindingProxy>& ObjectBindings, const FBakingAnimationKeySettings& InSettings, const FMovieSceneScriptingParams& Params)
 {
 	TSharedPtr<ISequencer> Sequencer = GetActiveSequencer();
 	if (Sequencer == nullptr)
 	{
-		UE_LOG(LogLevelSequenceEditor, Warning, TEXT("Bake Transform failed."));
 		return false;
 	}
 
-	UMovieScene* FocusedMovieScene = Sequencer->GetFocusedMovieSceneSequence()->GetMovieScene();
-	if (!FocusedMovieScene)
-	{
-		UE_LOG(LogLevelSequenceEditor, Warning, TEXT("Bake Transform failed."));
-		return false;
-	}
-
-	if (FocusedMovieScene->IsReadOnly())
-	{
-		UE_LOG(LogLevelSequenceEditor, Warning, TEXT("Bake Transform failed."));
-		FSequencerUtilities::ShowReadOnlyError();
-		return false;
-	}
-
-	if (ObjectBindings.Num() == 0)
-	{
-		UE_LOG(LogLevelSequenceEditor, Warning, TEXT("Bake Transform failed."));
-		return false;
-	}
-
-	FScopedTransaction BakeTransform(LOCTEXT("BakeTransform", "Bake Transform"));
-
-	FocusedMovieScene->Modify();
+	FFrameRate TickResolution = Sequencer->GetFocusedTickResolution();
+	FFrameRate DisplayRate = Sequencer->GetFocusedDisplayRate();
 
-	FQualifiedFrameTime ResetTime = Sequencer->GetLocalTime();
+	FBakingAnimationKeySettings Settings = InSettings;
 
-	FFrameRate TickResolution = FocusedMovieScene->GetTickResolution();
-	FFrameRate DisplayRate = FocusedMovieScene->GetDisplayRate();
-
-	FBakingAnimationKeySettings SettingsInTick = InSettings;
-	
 	if (Params.TimeUnit == EMovieSceneTimeUnit::DisplayRate)
 	{
-		SettingsInTick.StartFrame = ConvertFrameTime(SettingsInTick.StartFrame, DisplayRate, TickResolution).GetFrame();
-		SettingsInTick.EndFrame = ConvertFrameTime(SettingsInTick.EndFrame, DisplayRate, TickResolution).GetFrame();
-	}
-
-	TSortedMap<FFrameNumber, FFrameNumber> TotalFrameMap;
-	TMap<FGuid, FBakeData> BakeDataMap;
-	for (const FMovieSceneBindingProxy& ObjectBinding : ObjectBindings)
-	{
-		BakeDataMap.Add(ObjectBinding.BindingID);
+		Settings.StartFrame = ConvertFrameTime(Settings.StartFrame, DisplayRate, TickResolution).GetFrame();
+		Settings.EndFrame = ConvertFrameTime(Settings.EndFrame, DisplayRate, TickResolution).GetFrame();
 	}
-	CalculateFramesPerGuid(Sequencer, SettingsInTick, BakeDataMap, TotalFrameMap);
-
-	FMovieSceneInverseSequenceTransform LocalToRootTransform = Sequencer->GetFocusedMovieSceneSequenceTransform().Inverse();
-	
-	TArray<FFrameNumber> AllFrames;
-	TotalFrameMap.GenerateKeyArray(AllFrames);
-
-	UWorld* PlaybackContext = Sequencer->GetPlaybackContext()->GetWorld();
-	ensure(PlaybackContext);
-	const FConstraintsManagerController& Controller = FConstraintsManagerController::Get(PlaybackContext);
-
-	for (FFrameNumber KeyTime: AllFrames)
-	{
-		TOptional<FFrameTime> NewGlobalTime = LocalToRootTransform.TryTransformTime(KeyTime);
-		if (!NewGlobalTime)
-		{
-			continue;
-		}
-
-		Sequencer->SetGlobalTime(NewGlobalTime.GetValue());
-		Controller.EvaluateAllConstraints();
-
-		for (const FMovieSceneBindingProxy& ObjectBinding : ObjectBindings)
-		{
-			FGuid Guid = ObjectBinding.BindingID;
 
-			for (TWeakObjectPtr<> RuntimeObject : Sequencer->FindObjectsInCurrentSequence(Guid) )
-			{
-				const FFrameNumber* Number = BakeDataMap[Guid].KeyTimes.FindKey(KeyTime);
-				if (Number == nullptr)
-				{
-					continue;
-				}
-				AActor* Actor = Cast<AActor>(RuntimeObject.Get());
-				if (!Actor)
-				{
-					UActorComponent* ActorComponent = Cast<UActorComponent>(RuntimeObject.Get());
-					if (ActorComponent)
-					{
-						Actor = ActorComponent->GetOwner();
-					}
-				}
-
-				if (!Actor)
-				{
-					continue;
-				}
-
-				UCameraComponent* CameraComponent = MovieSceneHelpers::CameraComponentFromRuntimeObject(RuntimeObject.Get());
-
-				// Cache transforms
-				USceneComponent* Parent = nullptr;
-				if (CameraComponent)
-				{
-					Parent = CameraComponent->GetAttachParent();
-				} 
-				else if (Actor->GetRootComponent())
-				{
-					Parent = Actor->GetRootComponent()->GetAttachParent();
-				}
-				
-				// The CameraRig_rail updates the spline position tick, so it needs to be ticked manually while baking the frames
-				while (Parent && Parent->GetOwner())
-				{
-					Parent->GetOwner()->Tick(0.03f);
-					if (USkeletalMeshComponent* SkeletalMeshComponent = Cast<USkeletalMeshComponent>(Parent))
-					{
-						SkeletalMeshComponent->TickAnimation(0.f, false);
-
-						SkeletalMeshComponent->RefreshBoneTransforms();
-						SkeletalMeshComponent->RefreshFollowerComponents();
-						SkeletalMeshComponent->UpdateComponentToWorld();
-						SkeletalMeshComponent->FinalizeBoneTransform();
-						SkeletalMeshComponent->MarkRenderTransformDirty();
-						SkeletalMeshComponent->MarkRenderDynamicDataDirty();
-					}
-					Parent = Parent->GetAttachParent();
-				}
-
-				if (CameraComponent)
-				{
-					FTransform AdditiveOffset;
-					float AdditiveFOVOffset;
-					CameraComponent->GetAdditiveOffset(AdditiveOffset, AdditiveFOVOffset);
-
-					FTransform Transform(Actor->GetActorRotation(), Actor->GetActorLocation());
-					FTransform TransformWithAdditiveOffset = AdditiveOffset * Transform;
-					FVector LocalTranslation = TransformWithAdditiveOffset.GetTranslation();
-					FRotator LocalRotation = TransformWithAdditiveOffset.GetRotation().Rotator();
-
-					BakeDataMap[Guid].Locations.Add(LocalTranslation);
-					BakeDataMap[Guid].Rotations.Add(LocalRotation);
-					BakeDataMap[Guid].Scales.Add(FVector::OneVector);
-				}
-				else
-				{
-					BakeDataMap[Guid].Locations.Add(Actor->GetActorLocation());
-					BakeDataMap[Guid].Rotations.Add(Actor->GetActorRotation());
-					BakeDataMap[Guid].Scales.Add(Actor->GetActorScale());
-				}
-
-			}
-		}
-	}
-
-	const bool bDisableSectionsAfterBaking = Sequencer->GetSequencerSettings()->GetDisableSectionsAfterBaking();
-
-	for (TPair<FGuid, FBakeData>& BakeData : BakeDataMap)
-	{
-		FGuid Guid = BakeData.Key;
-		TArray<FFrameNumber> KeyTimes;
-		BakeData.Value.KeyTimes.GenerateKeyArray(KeyTimes);
-		// Disable or delete any constraint (attach/path) tracks
-		AActor* AttachParentActor = nullptr;
-		for (UMovieSceneTrack* Track : FocusedMovieScene->FindTracks(UMovieScene3DConstraintTrack::StaticClass(), Guid))
-		{
-			if (UMovieScene3DConstraintTrack* ConstraintTrack = Cast<UMovieScene3DConstraintTrack>(Track))
-			{
-				for (UMovieSceneSection* ConstraintSection : ConstraintTrack->GetAllSections())
-				{
-					FMovieSceneObjectBindingID ConstraintBindingID = (Cast<UMovieScene3DConstraintSection>(ConstraintSection))->GetConstraintBindingID();
-					if (auto BoundObjectsView = ConstraintBindingID.ResolveBoundObjects(Sequencer->GetFocusedTemplateID(), *Sequencer); BoundObjectsView.Num() > 0)
-					{
-						TWeakObjectPtr<> ParentObject = BoundObjectsView[0];
-						AttachParentActor = Cast<AActor>(ParentObject.Get());
-					}
-				}
-
-				if (bDisableSectionsAfterBaking)
-				{
-					for (UMovieSceneSection* ConstraintSection : ConstraintTrack->GetAllSections())
-					{
-						ConstraintSection->Modify();
-						ConstraintSection->SetIsActive(false);
-					}
-				}
-				else
-				{
-					FocusedMovieScene->RemoveTrack(*ConstraintTrack);
-				}
-			}
-		}
-
-		// Disable or delete any transform tracks
-		for (UMovieSceneTrack* Track : FocusedMovieScene->FindTracks(UMovieScene3DTransformTrack::StaticClass(), Guid))
-		{
-			if (UMovieScene3DTransformTrack* TransformTrack = Cast<UMovieScene3DTransformTrack>(Track))
-			{
-				if (bDisableSectionsAfterBaking)
-				{
-					for (UMovieSceneSection* TransformSection : TransformTrack->GetAllSections())
-					{
-						TransformSection->Modify();
-						TransformSection->SetIsActive(false);
-					}
-				}
-				else
-				{
-					FocusedMovieScene->RemoveTrack(*TransformTrack);
-				}
-			}
-		}
-
-		// Disable or delete any camera shake tracks
-		for (UMovieSceneTrack* Track : FocusedMovieScene->FindTracks(UMovieSceneCameraShakeTrack::StaticClass(), Guid))
-		{
-			if (UMovieSceneCameraShakeTrack* CameraShakeTrack = Cast<UMovieSceneCameraShakeTrack>(Track))
-			{
-				if (bDisableSectionsAfterBaking)
-				{
-					for (UMovieSceneSection* CameraShakeSection : CameraShakeTrack->GetAllSections())
-					{
-						CameraShakeSection->Modify();
-						CameraShakeSection->SetIsActive(false);
-					}
-				}
-				else
-				{
-					FocusedMovieScene->RemoveTrack(*CameraShakeTrack);
-				}
-			}
-		}
-
-		// Reset position
-		Sequencer->SetLocalTimeDirectly(ResetTime.Time);
-		Sequencer->ForceEvaluate();
-
-		FVector DefaultLocation = FVector::ZeroVector;
-		FVector DefaultRotation = FVector::ZeroVector;
-		FVector DefaultScale = FVector::OneVector;
-
-		for (TWeakObjectPtr<> RuntimeObject : Sequencer->FindObjectsInCurrentSequence(Guid))
-		{
-			AActor* Actor = Cast<AActor>(RuntimeObject.Get());
-			if (!Actor)
-			{
-				UActorComponent* ActorComponent = Cast<UActorComponent>(RuntimeObject.Get());
-				if (ActorComponent)
-				{
-					Actor = ActorComponent->GetOwner();
-				}
-			}
-
-			if (!Actor)
-			{
-				continue;
-			}
-
-			DefaultLocation = Actor->GetActorLocation();
-			DefaultRotation = Actor->GetActorRotation().Euler();
-			DefaultScale = Actor->GetActorScale();
-
-			// Always detach from any existing parent
-			Actor->DetachFromActor(FDetachmentTransformRules::KeepRelativeTransform);
-		}
-			
-		// Create new transform track and section
-		UMovieScene3DTransformTrack* TransformTrack = Cast<UMovieScene3DTransformTrack>(FocusedMovieScene->AddTrack(UMovieScene3DTransformTrack::StaticClass(), Guid));
-
-		if (TransformTrack)
-		{
-			UMovieScene3DTransformSection* TransformSection = CastChecked<UMovieScene3DTransformSection>(TransformTrack->CreateNewSection());
-			TransformTrack->AddSection(*TransformSection);
-
-			TransformSection->SetRange(TRange<FFrameNumber>::All());
-
-			TArrayView<FMovieSceneDoubleChannel*> DoubleChannels = TransformSection->GetChannelProxy().GetChannels<FMovieSceneDoubleChannel>();
-			DoubleChannels[0]->SetDefault(DefaultLocation.X);
-			DoubleChannels[1]->SetDefault(DefaultLocation.Y);
-			DoubleChannels[2]->SetDefault(DefaultLocation.Z);
-			DoubleChannels[3]->SetDefault(DefaultRotation.X);
-			DoubleChannels[4]->SetDefault(DefaultRotation.Y);
-			DoubleChannels[5]->SetDefault(DefaultRotation.Z);
-			DoubleChannels[6]->SetDefault(DefaultScale.X);
-			DoubleChannels[7]->SetDefault(DefaultScale.Y);
-			DoubleChannels[8]->SetDefault(DefaultScale.Z);
-
-			TArray<FVector> LocalTranslations, LocalRotations, LocalScales;
-			LocalTranslations.SetNum(KeyTimes.Num());
-			LocalRotations.SetNum(KeyTimes.Num());
-			LocalScales.SetNum(KeyTimes.Num());
-
-			for (int32 Counter = 0; Counter < KeyTimes.Num(); ++Counter)
-			{
-				FVector LocalTranslation = DefaultLocation;
-				FVector LocalScale = DefaultScale;
-				FRotator LocalRotation = DefaultRotation.Rotation();
-
-				if (Counter < BakeData.Value.Locations.Num())
-				{
-					LocalTranslation = BakeData.Value.Locations[Counter];
-				}
-				if (Counter < BakeData.Value.Rotations.Num())
-				{
-					LocalRotation = BakeData.Value.Rotations[Counter];
-				}
-				if (Counter < BakeData.Value.Scales.Num())
-				{
-					LocalScale = BakeData.Value.Scales[Counter];
-				}
-
-				FTransform LocalTransform(LocalRotation, LocalTranslation, LocalScale);
-				LocalTranslations[Counter] = LocalTransform.GetTranslation();
-				LocalRotations[Counter] = LocalTransform.GetRotation().Euler();
-				LocalScales[Counter] = LocalTransform.GetScale3D();
-			}
-
-			// Euler filter
-			for (int32 Counter = 0; Counter < LocalRotations.Num() - 1; ++Counter)
-			{
-				FMath::WindRelativeAnglesDegrees(LocalRotations[Counter].X, LocalRotations[Counter + 1].X);
-				FMath::WindRelativeAnglesDegrees(LocalRotations[Counter].Y, LocalRotations[Counter + 1].Y);
-				FMath::WindRelativeAnglesDegrees(LocalRotations[Counter].Z, LocalRotations[Counter + 1].Z);							
-			}
-			if (SettingsInTick.BakingKeySettings == EBakingKeySettings::KeysOnly)
-			{
-				const EMovieSceneKeyInterpolation KeyInterpolation = Sequencer->GetSequencerSettings()->GetKeyInterpolation();
-
-				for (int32 Counter = 0; Counter < KeyTimes.Num(); ++Counter)
-				{
-					int ChannelIndex = 0;
-					FFrameNumber KeyTime = KeyTimes[Counter];
-					{
-						TMovieSceneChannelData<FMovieSceneDoubleValue> ChannelData = DoubleChannels[ChannelIndex++]->GetData();
-						MovieSceneToolHelpers::SetOrAddKey(ChannelData, KeyTime, LocalTranslations[Counter].X);
-					}
-					{
-						TMovieSceneChannelData<FMovieSceneDoubleValue> ChannelData = DoubleChannels[ChannelIndex++]->GetData();
-						MovieSceneToolHelpers::SetOrAddKey(ChannelData, KeyTime, LocalTranslations[Counter].Y);
-					}
-					{
-						TMovieSceneChannelData<FMovieSceneDoubleValue> ChannelData = DoubleChannels[ChannelIndex++]->GetData();
-						MovieSceneToolHelpers::SetOrAddKey(ChannelData, KeyTime, LocalTranslations[Counter].Z);
-					}
-					{
-						TMovieSceneChannelData<FMovieSceneDoubleValue> ChannelData = DoubleChannels[ChannelIndex++]->GetData();
-						MovieSceneToolHelpers::SetOrAddKey(ChannelData, KeyTime, LocalRotations[Counter].X);
-					}
-					{
-						TMovieSceneChannelData<FMovieSceneDoubleValue> ChannelData = DoubleChannels[ChannelIndex++]->GetData();
-						MovieSceneToolHelpers::SetOrAddKey(ChannelData, KeyTime, LocalRotations[Counter].Y);
-
-					}
-					{
-						TMovieSceneChannelData<FMovieSceneDoubleValue> ChannelData = DoubleChannels[ChannelIndex++]->GetData();
-						MovieSceneToolHelpers::SetOrAddKey(ChannelData, KeyTime, LocalRotations[Counter].Z);
-					}
-					{
-						TMovieSceneChannelData<FMovieSceneDoubleValue> ChannelData = DoubleChannels[ChannelIndex++]->GetData();
-						MovieSceneToolHelpers::SetOrAddKey(ChannelData, KeyTime, LocalScales[Counter].X);
-					}
-					{
-						TMovieSceneChannelData<FMovieSceneDoubleValue> ChannelData = DoubleChannels[ChannelIndex++]->GetData();
-						MovieSceneToolHelpers::SetOrAddKey(ChannelData, KeyTime, LocalScales[Counter].Y);
-					}
-					{
-						TMovieSceneChannelData<FMovieSceneDoubleValue> ChannelData = DoubleChannels[ChannelIndex++]->GetData();
-						MovieSceneToolHelpers::SetOrAddKey(ChannelData, KeyTime, LocalScales[Counter].Z);
-					}
-				}
-			}
-			else
-			{
-				for (int32 Counter = 0; Counter < KeyTimes.Num(); ++Counter)
-				{
-					FFrameNumber KeyTime = KeyTimes[Counter];
-					DoubleChannels[0]->AddLinearKey(KeyTime, LocalTranslations[Counter].X);
-					DoubleChannels[1]->AddLinearKey(KeyTime, LocalTranslations[Counter].Y);
-					DoubleChannels[2]->AddLinearKey(KeyTime, LocalTranslations[Counter].Z);
-					DoubleChannels[3]->AddLinearKey(KeyTime, LocalRotations[Counter].X);
-					DoubleChannels[4]->AddLinearKey(KeyTime, LocalRotations[Counter].Y);
-					DoubleChannels[5]->AddLinearKey(KeyTime, LocalRotations[Counter].Z);
-					DoubleChannels[6]->AddLinearKey(KeyTime, LocalScales[Counter].X);
-					DoubleChannels[7]->AddLinearKey(KeyTime, LocalScales[Counter].Y);
-					DoubleChannels[8]->AddLinearKey(KeyTime, LocalScales[Counter].Z);
-				}
-				if (SettingsInTick.bReduceKeys == true)
-				{
-					FKeyDataOptimizationParams Param;
-					Param.bAutoSetInterpolation = true;
-					Param.Tolerance = SettingsInTick.Tolerance;
-					TRange<FFrameNumber> Range(SettingsInTick.StartFrame, SettingsInTick.EndFrame);
-					Param.Range = Range;
-					MovieSceneToolHelpers::OptimizeSection(Param, TransformSection);
-				}
-			}
-		}
-	}
-	
-	Sequencer->NotifyMovieSceneDataChanged( EMovieSceneDataChangeType::MovieSceneStructureItemsChanged );
-	return true;
+	return FSequencerUtilities::BakeTransform(Sequencer.ToSharedRef(), ObjectBindings, InSettings);
 }
 
 void ULevelSequenceEditorSubsystem::FixActorReferences()
diff --git a/Engine/Plugins/MovieScene/LevelSequenceEditor/Source/LevelSequenceEditor/Public/LevelSequenceEditorSubsystem.h b/Engine/Plugins/MovieScene/LevelSequenceEditor/Source/LevelSequenceEditor/Public/LevelSequenceEditorSubsystem.h
--- a/Engine/Plugins/MovieScene/LevelSequenceEditor/Source/LevelSequenceEditor/Public/LevelSequenceEditorSubsystem.h
+++ b/Engine/Plugins/MovieScene/LevelSequenceEditor/Source/LevelSequenceEditor/Public/LevelSequenceEditorSubsystem.h
@@ -279,16 +279,6 @@ public:
 	UE_API void RefreshTrackRowMetadataDetails(IDetailsView* DetailsView);
 
 private:
-	/** Used by Baking transforms*/
-	struct FBakeData
-	{
-		TArray<FVector> Locations;
-		TArray<FRotator> Rotations;
-		TArray<FVector> Scales;
-		TSortedMap<FFrameNumber,FFrameNumber> KeyTimes;
-	};
-	UE_API void CalculateFramesPerGuid(TSharedPtr<ISequencer>& Sequencer, const FBakingAnimationKeySettings& InSettings, TMap<FGuid, FBakeData>& OutBakeDataMa,
-		TSortedMap<FFrameNumber, FFrameNumber>&  OutFrameMap);
 
 	// Used by binding properties menu
 	struct FBindingPropertiesNotifyHook : FNotifyHook
@@ -321,7 +311,6 @@ private:
 	
 	UE_API void SnapSectionsToTimelineUsingSourceTimecodeInternal();
 	UE_API void SyncSectionsUsingSourceTimecodeInternal();
-	UE_API void BakeTransformInternal();
 	UE_API void AddActorsToBindingInternal();
 	UE_API void ReplaceBindingWithActorsInternal();
 	UE_API void RemoveActorsFromBindingInternal();
diff --git a/Engine/Source/Editor/Sequencer/Private/SSequencer.cpp b/Engine/Source/Editor/Sequencer/Private/SSequencer.cpp
--- a/Engine/Source/Editor/Sequencer/Private/SSequencer.cpp
+++ b/Engine/Source/Editor/Sequencer/Private/SSequencer.cpp
@@ -1,6 +1,7 @@
 // Copyright Epic Games, Inc. All Rights Reserved.
 
 #include "SSequencer.h"
+#include "BakingAnimationKeySettings.h"
 #include "Engine/Blueprint.h"
 #include "Filters/Filters/SequencerTrackFilter_CustomText.h"
 #include "Filters/Filters/SequencerTrackFilters.h"
@@ -69,6 +70,7 @@
 #include "ScopedTransaction.h"
 #include "SequencerTimeSliderController.h"
 #include "SequencerToolMenuContext.h"
+#include "SSequencerBakeTransform.h"
 #include "SSequencerSectionOverlay.h"
 #include "STemporarilyFocusedSpinBox.h"
 #include "MVVM/Views/STrackAreaView.h"
@@ -1090,6 +1092,11 @@ void SSequencer::BindCommands(TSharedRef<FUICommandList> SequencerCommandBinding
 		FExecuteAction::CreateLambda([this] { TransformBox->ToggleVisibility(); })
 	);
 
+	SequencerCommandBindings->MapAction(
+		FSequencerCommands::Get().BakeTransform,
+		FExecuteAction::CreateSP(this, &SSequencer::BakeTransform)
+	);
+
 	// Allow jumping to the Sequencer tree search if you have Sequencer focused
 	SequencerCommandBindings->MapAction(
 		FSequencerCommands::Get().QuickTreeSearch,
@@ -1744,6 +1751,11 @@ TSharedRef<SWidget> SSequencer::MakeActionsMenu()
 	// transform actions
 	MenuBuilder.BeginSection("Transform", LOCTEXT("TransformHeader", "Transform"));
 	
+	if (SequencerPtr.Pin()->IsLevelEditorSequencer())
+	{
+		MenuBuilder.AddMenuEntry(FSequencerCommands::Get().BakeTransform);
+	}
+	
 	MenuBuilder.AddMenuEntry(FSequencerCommands::Get().ToggleShowTransformBox);
 	MenuBuilder.AddMenuEntry(FSequencerCommands::Get().ToggleShowStretchBox);
 	MenuBuilder.AddMenuEntry(FSequencerCommands::Get().TranslateLeft);
@@ -3764,6 +3776,51 @@ void SSequencer::OnSetSequenceReadOnly(ECheckBoxState CheckBoxState)
 	}
 }
 
+void SSequencer::BakeTransform()
+{
+	TSharedPtr<FSequencer> Sequencer = SequencerPtr.Pin();
+	if (!Sequencer)
+	{
+		return;
+	}
+
+	UMovieSceneSequence* Sequence = Sequencer->GetFocusedMovieSceneSequence();
+	UMovieScene* MovieScene = Sequence ? Sequence->GetMovieScene() : nullptr;
+	if (!MovieScene)
+	{
+		return;
+	}
+
+	TArray<FGuid> ObjectBindings;
+	Sequencer->GetSelectedObjects(ObjectBindings);
+
+	TArray<FMovieSceneBindingProxy> BindingProxies;
+	for (FGuid Guid : ObjectBindings)
+	{
+		BindingProxies.Add(FMovieSceneBindingProxy(Guid, Sequence));
+	}
+
+	static FBakingAnimationKeySettings Settings; //reuse the settings except for the range
+	Settings.StartFrame = UE::MovieScene::DiscreteInclusiveLower(MovieScene->GetPlaybackRange());
+	Settings.EndFrame = UE::MovieScene::DiscreteExclusiveUpper(MovieScene->GetPlaybackRange());
+
+	TSharedRef<SBakeTransformWidget> BakeWidget =
+		SNew(SBakeTransformWidget)
+		.Settings(Settings)
+		.Sequencer(Sequencer.Get())
+		.OnBake_Lambda([this, &BindingProxies, Sequencer](FBakingAnimationKeySettings InSettings)
+		{
+			if (FSequencerUtilities::BakeTransform(Sequencer.ToSharedRef(), BindingProxies, InSettings))
+			{
+				Settings = InSettings;
+			}
+			
+			return FReply::Handled();
+		});
+
+	BakeWidget->OpenDialog(true);
+}
+
 void SSequencer::SetPlayTimeClampedByWorkingRange(double Frame)
 {
 	if (SequencerPtr.IsValid())
diff --git a/Engine/Source/Editor/Sequencer/Private/SSequencer.h b/Engine/Source/Editor/Sequencer/Private/SSequencer.h
--- a/Engine/Source/Editor/Sequencer/Private/SSequencer.h
+++ b/Engine/Source/Editor/Sequencer/Private/SSequencer.h
@@ -595,6 +595,8 @@ private:
 	/** Returns whether or not the Curve Editor is enabled. Allows us to bind to the Slate Enabled attribute. */
 	bool GetIsCurveEditorEnabled() const { return !GetIsSequenceReadOnly(); }
 
+	void BakeTransform();
+
 public:
 	/** On Paste Command */
 	void OnPaste();
diff --git a/Engine/Source/Editor/Sequencer/Private/SSequencerBakeTransform.cpp b/Engine/Source/Editor/Sequencer/Private/SSequencerBakeTransform.cpp
new file mode 100644
--- /dev/null
+++ b/Engine/Source/Editor/Sequencer/Private/SSequencerBakeTransform.cpp
@@ -0,0 +1,162 @@
+// Copyright Epic Games, Inc. All Rights Reserved.
+
+#include "SSequencerBakeTransform.h"
+
+#include "Editor.h"
+#include "Framework/Application/SlateApplication.h"
+#include "ISequencer.h"
+#include "IStructureDetailsView.h"
+#include "PropertyEditorModule.h"
+#include "Widgets/SWindow.h"
+#include "Widgets/Input/SButton.h"
+#include "Widgets/Layout/SSpacer.h"
+
+#define LOCTEXT_NAMESPACE "SBakeTransformWidget"
+
+void SBakeTransformWidget::Construct(const FArguments& InArgs)
+{
+	check(InArgs._Sequencer);
+	check(InArgs._OnBake.IsBound());
+
+	Settings = MakeShared<TStructOnScope<FBakingAnimationKeySettings>>();
+	Settings->InitializeAs<FBakingAnimationKeySettings>();
+	*Settings = InArgs._Settings;
+	//always setting space to be parent as default, since stored space may not be available.
+	Sequencer = InArgs._Sequencer;
+
+	FStructureDetailsViewArgs StructureViewArgs;
+	StructureViewArgs.bShowObjects = true;
+	StructureViewArgs.bShowAssets = true;
+	StructureViewArgs.bShowClasses = true;
+	StructureViewArgs.bShowInterfaces = true;
+
+	FDetailsViewArgs ViewArgs;
+	ViewArgs.bAllowSearch = false;
+	ViewArgs.bHideSelectionTip = false;
+	ViewArgs.bShowObjectLabel = false;
+
+	FPropertyEditorModule& PropertyEditor = FModuleManager::Get().LoadModuleChecked<FPropertyEditorModule>(TEXT("PropertyEditor"));
+
+	DetailsView = PropertyEditor.CreateStructureDetailView(ViewArgs, StructureViewArgs, TSharedPtr<FStructOnScope>());
+
+	DetailsView->GetDetailsView()->RegisterInstancedCustomPropertyTypeLayout("FrameNumber",
+		FOnGetPropertyTypeCustomizationInstance::CreateSP(Sequencer, &ISequencer::MakeFrameNumberDetailsCustomization));
+	DetailsView->SetStructureData(Settings);
+
+	ChildSlot
+		[
+			SNew(SBorder)
+				.Visibility(EVisibility::Visible)
+				[
+					SNew(SVerticalBox)
+
+
+						+ SVerticalBox::Slot()
+						.AutoHeight()
+						.Padding(0.f, 0.f, 0.f, 0.f)
+						[
+							DetailsView->GetWidget().ToSharedRef()
+						]
+
+						+ SVerticalBox::Slot()
+						.AutoHeight()
+						.Padding(0.f, 16.f, 0.f, 16.f)
+						[
+							SNew(SHorizontalBox)
+
+								+ SHorizontalBox::Slot()
+								.FillWidth(1.f)
+								[
+									SNew(SSpacer)
+								]
+
+								+ SHorizontalBox::Slot()
+								.AutoWidth()
+								.HAlign(HAlign_Left)
+								.VAlign(VAlign_Center)
+								.Padding(8.f, 0.f, 0.f, 0.f)
+								[
+									SNew(SButton)
+										.HAlign(HAlign_Center)
+										.ContentPadding(FAppStyle::GetMargin("StandardDialog.ContentPadding"))
+										.Text(LOCTEXT("OK", "OK"))
+										.OnClicked_Lambda([this, InArgs]()
+											{
+												FReply Reply = InArgs._OnBake.Execute(*(Settings->Get()));
+												CloseDialog();
+												return Reply;
+
+											})
+								]
+
+								+ SHorizontalBox::Slot()
+								.AutoWidth()
+								.HAlign(HAlign_Left)
+								.VAlign(VAlign_Center)
+								.Padding(8.f, 0.f, 16.f, 0.f)
+								[
+									SNew(SButton)
+										.HAlign(HAlign_Center)
+										.ContentPadding(FAppStyle::GetMargin("StandardDialog.ContentPadding"))
+										.Text(LOCTEXT("Cancel", "Cancel"))
+										.OnClicked_Lambda([this]()
+											{
+												CloseDialog();
+												return FReply::Handled();
+											})
+								]
+						]
+				]
+		];
+}
+
+class SBakeTransformDialogWindow : public SWindow
+{
+};
+
+FReply SBakeTransformWidget::OpenDialog(bool bModal)
+{
+	check(!DialogWindow.IsValid());
+
+	const FVector2D CursorPos = FSlateApplication::Get().GetCursorPos();
+
+	TSharedRef<SBakeTransformDialogWindow> Window = SNew(SBakeTransformDialogWindow)
+		.Title(LOCTEXT("SBakeTransformWidgetTitle", "Bake Transform"))
+		.CreateTitleBar(true)
+		.Type(EWindowType::Normal)
+		.SizingRule(ESizingRule::Autosized)
+		.ScreenPosition(CursorPos)
+		.FocusWhenFirstShown(true)
+		.ActivationPolicy(EWindowActivationPolicy::FirstShown)
+		[
+			AsShared()
+		];
+
+	Window->SetWidgetToFocusOnActivate(AsShared());
+
+	DialogWindow = Window;
+
+	Window->MoveWindowTo(CursorPos);
+
+	if (bModal)
+	{
+		GEditor->EditorAddModalWindow(Window);
+	}
+	else
+	{
+		FSlateApplication::Get().AddWindow(Window);
+	}
+
+	return FReply::Handled();
+}
+
+void SBakeTransformWidget::CloseDialog()
+{
+	if (DialogWindow.IsValid())
+	{
+		DialogWindow.Pin()->RequestDestroyWindow();
+		DialogWindow.Reset();
+	}
+}
+
+#undef LOCTEXT_NAMESPACE
\ No newline at end of file
diff --git a/Engine/Source/Editor/Sequencer/Private/SSequencerBakeTransform.h b/Engine/Source/Editor/Sequencer/Private/SSequencerBakeTransform.h
new file mode 100644
--- /dev/null
+++ b/Engine/Source/Editor/Sequencer/Private/SSequencerBakeTransform.h
@@ -0,0 +1,44 @@
+// Copyright Epic Games, Inc. All Rights Reserved.
+
+#pragma once
+
+#include "BakingAnimationKeySettings.h"
+#include "Widgets/DeclarativeSyntaxSupport.h"
+#include "Widgets/SCompoundWidget.h"
+#include "UObject/StructOnScope.h"
+
+class ISequencer;
+class IStructureDetailsView;
+class SWindow;
+
+DECLARE_DELEGATE_RetVal_OneParam(FReply, SBakeTransformOnBake, FBakingAnimationKeySettings);
+
+/** Widget allowing baking controls from one space to another */
+class SBakeTransformWidget : public SCompoundWidget
+{
+public:
+
+	SLATE_BEGIN_ARGS(SBakeTransformWidget)
+		: _Sequencer(nullptr)
+		{}
+		SLATE_ARGUMENT(ISequencer*, Sequencer)
+		SLATE_ARGUMENT(FBakingAnimationKeySettings, Settings)
+		SLATE_EVENT(SBakeTransformOnBake, OnBake)
+	SLATE_END_ARGS()
+
+	void Construct(const FArguments& InArgs);
+	virtual ~SBakeTransformWidget() override {}
+
+	FReply OpenDialog(bool bModal = true);
+	void CloseDialog();
+
+private:
+
+	//used for setting up the details
+	TSharedPtr<TStructOnScope<FBakingAnimationKeySettings>> Settings;
+
+	ISequencer* Sequencer;
+
+	TWeakPtr<SWindow> DialogWindow;
+	TSharedPtr<IStructureDetailsView> DetailsView;
+};
diff --git a/Engine/Source/Editor/Sequencer/Private/SequencerCommands.cpp b/Engine/Source/Editor/Sequencer/Private/SequencerCommands.cpp
--- a/Engine/Source/Editor/Sequencer/Private/SequencerCommands.cpp
+++ b/Engine/Source/Editor/Sequencer/Private/SequencerCommands.cpp
@@ -197,6 +197,8 @@ void FSequencerCommands::RegisterCommands()
 
 	UI_COMMAND(ToggleLimitViewportSelection, "Limit Viewport Selection", "Toggles viewport selectability between Sequencer only objects and all objects.", EUserInterfaceActionType::ToggleButton, FInputChord());
 
+	UI_COMMAND(BakeTransform, "Bake Transform", "Bake transform in world space, removing any existing transform and attach tracks", EUserInterfaceActionType::Button, FInputChord());
+
 	UI_COMMAND(ToggleSidebarVisible, "Sidebar Details", "Toggles the visibility of the Sequencer sidebar.", EUserInterfaceActionType::ToggleButton, FInputChord());
 	UI_COMMAND(ToggleSidebarSelectionDrawerOpen, "Sidebar Selection Drawer Open", "Toggles the visible state of the sidebar selection drawer. Undocks and closes the sidebar drawer if docked.", EUserInterfaceActionType::Button, FInputChord(EModifierKey::Control, EKeys::E));
 	UI_COMMAND(ToggleSidebarDrawerDock, "Sidebar Drawer Docked", "Toggles the docked state of the sidebar current drawer. Undocks the docked sidebar drawer if docked or docks the sidebar drawer if there is one open and no currently docked drawer.", EUserInterfaceActionType::Button, FInputChord(EModifierKey::Alt, EKeys::D));
diff --git a/Engine/Source/Editor/Sequencer/Private/SequencerUtilities.cpp b/Engine/Source/Editor/Sequencer/Private/SequencerUtilities.cpp
--- a/Engine/Source/Editor/Sequencer/Private/SequencerUtilities.cpp
+++ b/Engine/Source/Editor/Sequencer/Private/SequencerUtilities.cpp
@@ -1,10 +1,13 @@
 // Copyright Epic Games, Inc. All Rights Reserved.
 
 #include "SequencerUtilities.h"
+#include "ActorForWorldTransforms.h"
 #include "AnimatedRange.h"
+#include "BakingAnimationKeySettings.h"
 #include "CineCameraActor.h"
 #include "CameraRig_Rail.h"
 #include "CameraRig_Crane.h"
+#include "Components/SkeletalMeshComponent.h"
 #include "Components/SplineComponent.h"
 #include "Containers/ArrayBuilder.h"
 #include "Editor/EditorEngine.h"
@@ -36,8 +39,12 @@
 #include "MovieSceneTimeHelpers.h"
 #include "MovieSceneToolHelpers.h"
 #include "MovieSceneTrack.h"
+#include "Sections/MovieScene3DConstraintSection.h"
 #include "Tracks/MovieScene3DAttachTrack.h"
+#include "Tracks/MovieScene3DConstraintTrack.h"
+#include "Tracks/MovieScene3DTransformTrack.h"
 #include "Tracks/MovieSceneCameraCutTrack.h"
+#include "Tracks/MovieSceneCameraShakeTrack.h"
 #include "Tracks/MovieSceneSpawnTrack.h"
 #include "Compilation/MovieSceneCompiledDataManager.h"
 #include "Framework/MultiBox/MultiBoxBuilder.h"
@@ -4407,6 +4414,452 @@ void FSequencerUtilities::RemoveActorsFromBinding(TSharedRef<ISequencer> Sequenc
 	Sequencer->NotifyMovieSceneDataChanged(EMovieSceneDataChangeType::MovieSceneStructureItemsChanged);
 }
 
+struct FBakeData
+{
+	TArray<FVector> Locations;
+	TArray<FRotator> Rotations;
+	TArray<FVector> Scales;
+	TSortedMap<FFrameNumber, FFrameNumber> KeyTimes;
+};
+
+void CalculateFramesPerGuid(TSharedRef<ISequencer> Sequencer, const FBakingAnimationKeySettings& InSettings, TMap<FGuid, FBakeData>& OutBakeDataMap, TSortedMap<FFrameNumber, FFrameNumber>& OutFrameMap)
+{
+	OutFrameMap.Reset();
+	TArray<FFrameNumber> Frames;
+	//we get all frames since we need to get the Actor PER FRAME in order to handle spanwables
+	MovieSceneToolHelpers::CalculateFramesBetween(Sequencer->GetFocusedMovieSceneSequence()->GetMovieScene(), InSettings.StartFrame, InSettings.EndFrame, InSettings.FrameIncrement, Frames);
+	if (InSettings.BakingKeySettings == EBakingKeySettings::AllFrames)
+	{
+		for (FFrameNumber& Frame : Frames)
+		{
+			OutFrameMap.Add(Frame, Frame);
+		}
+		for (TPair<FGuid, FBakeData>& BakeData : OutBakeDataMap)
+		{
+			BakeData.Value.KeyTimes.Reset();
+			BakeData.Value.KeyTimes = OutFrameMap;
+		}
+	}
+	else
+	{
+		for (TPair<FGuid, FBakeData>& BakeData : OutBakeDataMap)
+		{
+			FActorForWorldTransforms ActorForWorldTransforms;
+			FGuid Guid = BakeData.Key;
+
+			for (TWeakObjectPtr<> RuntimeObject : Sequencer->FindObjectsInCurrentSequence(Guid))
+			{
+				UActorComponent* ActorComponent = nullptr;
+				AActor* Actor = Cast<AActor>(RuntimeObject.Get());
+				if (!Actor)
+				{
+					ActorComponent = Cast<UActorComponent>(RuntimeObject.Get());
+					if (ActorComponent)
+					{
+						Actor = ActorComponent->GetOwner();
+					}
+				}
+				ActorForWorldTransforms.Actor = Actor;
+				ActorForWorldTransforms.Component = Cast<USceneComponent>(ActorComponent);
+				BakeData.Value.KeyTimes.Reset();
+				MovieSceneToolHelpers::GetActorsAndParentsKeyFrames(&Sequencer.Get(), ActorForWorldTransforms,
+					InSettings.StartFrame, InSettings.EndFrame, BakeData.Value.KeyTimes);
+				OutFrameMap.Append(BakeData.Value.KeyTimes);
+			}
+		}
+	}
+}
+
+bool FSequencerUtilities::BakeTransform(TSharedRef<ISequencer> Sequencer, const TArray<FMovieSceneBindingProxy>& ObjectBindings, const FBakingAnimationKeySettings& InSettings)
+{
+	UMovieScene* FocusedMovieScene = Sequencer->GetFocusedMovieSceneSequence()->GetMovieScene();
+	if (!FocusedMovieScene)
+	{
+		UE_LOG(LogSequencer, Warning, TEXT("Bake Transform failed."));
+		return false;
+	}
+
+	if (FocusedMovieScene->IsReadOnly())
+	{
+		UE_LOG(LogSequencer, Warning, TEXT("Bake Transform failed."));
+		FSequencerUtilities::ShowReadOnlyError();
+		return false;
+	}
+
+	if (ObjectBindings.Num() == 0)
+	{
+		UE_LOG(LogSequencer, Warning, TEXT("Bake Transform failed."));
+		return false;
+	}
+
+	FScopedTransaction BakeTransform(LOCTEXT("BakeTransform", "Bake Transform"));
+
+	FocusedMovieScene->Modify();
+
+	FQualifiedFrameTime ResetTime = Sequencer->GetLocalTime();
+
+	FFrameRate TickResolution = FocusedMovieScene->GetTickResolution();
+	FFrameRate DisplayRate = FocusedMovieScene->GetDisplayRate();
+
+	TSortedMap<FFrameNumber, FFrameNumber> TotalFrameMap;
+	TMap<FGuid, FBakeData> BakeDataMap;
+	for (const FMovieSceneBindingProxy& ObjectBinding : ObjectBindings)
+	{
+		BakeDataMap.Add(ObjectBinding.BindingID);
+	}
+	CalculateFramesPerGuid(Sequencer, InSettings, BakeDataMap, TotalFrameMap);
+
+	FMovieSceneInverseSequenceTransform LocalToRootTransform = Sequencer->GetFocusedMovieSceneSequenceTransform().Inverse();
+
+	TArray<FFrameNumber> AllFrames;
+	TotalFrameMap.GenerateKeyArray(AllFrames);
+
+	UWorld* PlaybackContext = Sequencer->GetPlaybackContext()->GetWorld();
+	ensure(PlaybackContext);
+	const FConstraintsManagerController& Controller = FConstraintsManagerController::Get(PlaybackContext);
+
+	for (FFrameNumber KeyTime : AllFrames)
+	{
+		TOptional<FFrameTime> NewGlobalTime = LocalToRootTransform.TryTransformTime(KeyTime);
+		if (!NewGlobalTime)
+		{
+			continue;
+		}
+
+		Sequencer->SetGlobalTime(NewGlobalTime.GetValue());
+		Controller.EvaluateAllConstraints();
+
+		for (const FMovieSceneBindingProxy& ObjectBinding : ObjectBindings)
+		{
+			FGuid Guid = ObjectBinding.BindingID;
+
+			for (TWeakObjectPtr<> RuntimeObject : Sequencer->FindObjectsInCurrentSequence(Guid))
+			{
+				const FFrameNumber* Number = BakeDataMap[Guid].KeyTimes.FindKey(KeyTime);
+				if (Number == nullptr)
+				{
+					continue;
+				}
+				AActor* Actor = Cast<AActor>(RuntimeObject.Get());
+				if (!Actor)
+				{
+					UActorComponent* ActorComponent = Cast<UActorComponent>(RuntimeObject.Get());
+					if (ActorComponent)
+					{
+						Actor = ActorComponent->GetOwner();
+					}
+				}
+
+				if (!Actor)
+				{
+					continue;
+				}
+
+				UCameraComponent* CameraComponent = MovieSceneHelpers::CameraComponentFromRuntimeObject(RuntimeObject.Get());
+
+				// Cache transforms
+				USceneComponent* Parent = nullptr;
+				if (CameraComponent)
+				{
+					Parent = CameraComponent->GetAttachParent();
+				}
+				else if (Actor->GetRootComponent())
+				{
+					Parent = Actor->GetRootComponent()->GetAttachParent();
+				}
+
+				// The CameraRig_rail updates the spline position tick, so it needs to be ticked manually while baking the frames
+				while (Parent && Parent->GetOwner())
+				{
+					Parent->GetOwner()->Tick(0.03f);
+					if (USkeletalMeshComponent* SkeletalMeshComponent = Cast<USkeletalMeshComponent>(Parent))
+					{
+						SkeletalMeshComponent->TickAnimation(0.f, false);
+
+						SkeletalMeshComponent->RefreshBoneTransforms();
+						SkeletalMeshComponent->RefreshFollowerComponents();
+						SkeletalMeshComponent->UpdateComponentToWorld();
+						SkeletalMeshComponent->FinalizeBoneTransform();
+						SkeletalMeshComponent->MarkRenderTransformDirty();
+						SkeletalMeshComponent->MarkRenderDynamicDataDirty();
+					}
+					Parent = Parent->GetAttachParent();
+				}
+
+				if (CameraComponent)
+				{
+					FTransform AdditiveOffset;
+					float AdditiveFOVOffset;
+					CameraComponent->GetAdditiveOffset(AdditiveOffset, AdditiveFOVOffset);
+
+					FTransform Transform(Actor->GetActorRotation(), Actor->GetActorLocation());
+					FTransform TransformWithAdditiveOffset = AdditiveOffset * Transform;
+					FVector LocalTranslation = TransformWithAdditiveOffset.GetTranslation();
+					FRotator LocalRotation = TransformWithAdditiveOffset.GetRotation().Rotator();
+
+					BakeDataMap[Guid].Locations.Add(LocalTranslation);
+					BakeDataMap[Guid].Rotations.Add(LocalRotation);
+					BakeDataMap[Guid].Scales.Add(FVector::OneVector);
+				}
+				else
+				{
+					BakeDataMap[Guid].Locations.Add(Actor->GetActorLocation());
+					BakeDataMap[Guid].Rotations.Add(Actor->GetActorRotation());
+					BakeDataMap[Guid].Scales.Add(Actor->GetActorScale());
+				}
+
+			}
+		}
+	}
+
+	const bool bDisableSectionsAfterBaking = Sequencer->GetSequencerSettings()->GetDisableSectionsAfterBaking();
+
+	for (TPair<FGuid, FBakeData>& BakeData : BakeDataMap)
+	{
+		FGuid Guid = BakeData.Key;
+		TArray<FFrameNumber> KeyTimes;
+		BakeData.Value.KeyTimes.GenerateKeyArray(KeyTimes);
+		// Disable or delete any constraint (attach/path) tracks
+		AActor* AttachParentActor = nullptr;
+		for (UMovieSceneTrack* Track : FocusedMovieScene->FindTracks(UMovieScene3DConstraintTrack::StaticClass(), Guid))
+		{
+			if (UMovieScene3DConstraintTrack* ConstraintTrack = Cast<UMovieScene3DConstraintTrack>(Track))
+			{
+				for (UMovieSceneSection* ConstraintSection : ConstraintTrack->GetAllSections())
+				{
+					FMovieSceneObjectBindingID ConstraintBindingID = (Cast<UMovieScene3DConstraintSection>(ConstraintSection))->GetConstraintBindingID();
+					if (auto BoundObjectsView = ConstraintBindingID.ResolveBoundObjects(Sequencer->GetFocusedTemplateID(), *Sequencer); BoundObjectsView.Num() > 0)
+					{
+						TWeakObjectPtr<> ParentObject = BoundObjectsView[0];
+						AttachParentActor = Cast<AActor>(ParentObject.Get());
+					}
+				}
+
+				if (bDisableSectionsAfterBaking)
+				{
+					for (UMovieSceneSection* ConstraintSection : ConstraintTrack->GetAllSections())
+					{
+						ConstraintSection->Modify();
+						ConstraintSection->SetIsActive(false);
+					}
+				}
+				else
+				{
+					FocusedMovieScene->RemoveTrack(*ConstraintTrack);
+				}
+			}
+		}
+
+		// Disable or delete any transform tracks
+		for (UMovieSceneTrack* Track : FocusedMovieScene->FindTracks(UMovieScene3DTransformTrack::StaticClass(), Guid))
+		{
+			if (UMovieScene3DTransformTrack* TransformTrack = Cast<UMovieScene3DTransformTrack>(Track))
+			{
+				if (bDisableSectionsAfterBaking)
+				{
+					for (UMovieSceneSection* TransformSection : TransformTrack->GetAllSections())
+					{
+						TransformSection->Modify();
+						TransformSection->SetIsActive(false);
+					}
+				}
+				else
+				{
+					FocusedMovieScene->RemoveTrack(*TransformTrack);
+				}
+			}
+		}
+
+		// Disable or delete any camera shake tracks
+		for (UMovieSceneTrack* Track : FocusedMovieScene->FindTracks(UMovieSceneCameraShakeTrack::StaticClass(), Guid))
+		{
+			if (UMovieSceneCameraShakeTrack* CameraShakeTrack = Cast<UMovieSceneCameraShakeTrack>(Track))
+			{
+				if (bDisableSectionsAfterBaking)
+				{
+					for (UMovieSceneSection* CameraShakeSection : CameraShakeTrack->GetAllSections())
+					{
+						CameraShakeSection->Modify();
+						CameraShakeSection->SetIsActive(false);
+					}
+				}
+				else
+				{
+					FocusedMovieScene->RemoveTrack(*CameraShakeTrack);
+				}
+			}
+		}
+
+		// Reset position
+		Sequencer->SetLocalTimeDirectly(ResetTime.Time);
+		Sequencer->ForceEvaluate();
+
+		FVector DefaultLocation = FVector::ZeroVector;
+		FVector DefaultRotation = FVector::ZeroVector;
+		FVector DefaultScale = FVector::OneVector;
+
+		for (TWeakObjectPtr<> RuntimeObject : Sequencer->FindObjectsInCurrentSequence(Guid))
+		{
+			AActor* Actor = Cast<AActor>(RuntimeObject.Get());
+			if (!Actor)
+			{
+				UActorComponent* ActorComponent = Cast<UActorComponent>(RuntimeObject.Get());
+				if (ActorComponent)
+				{
+					Actor = ActorComponent->GetOwner();
+				}
+			}
+
+			if (!Actor)
+			{
+				continue;
+			}
+
+			DefaultLocation = Actor->GetActorLocation();
+			DefaultRotation = Actor->GetActorRotation().Euler();
+			DefaultScale = Actor->GetActorScale();
+
+			// Always detach from any existing parent
+			Actor->DetachFromActor(FDetachmentTransformRules::KeepRelativeTransform);
+		}
+
+		// Create new transform track and section
+		UMovieScene3DTransformTrack* TransformTrack = Cast<UMovieScene3DTransformTrack>(FocusedMovieScene->AddTrack(UMovieScene3DTransformTrack::StaticClass(), Guid));
+
+		if (TransformTrack)
+		{
+			UMovieScene3DTransformSection* TransformSection = CastChecked<UMovieScene3DTransformSection>(TransformTrack->CreateNewSection());
+			TransformTrack->AddSection(*TransformSection);
+
+			TransformSection->SetRange(TRange<FFrameNumber>::All());
+
+			TArrayView<FMovieSceneDoubleChannel*> DoubleChannels = TransformSection->GetChannelProxy().GetChannels<FMovieSceneDoubleChannel>();
+			DoubleChannels[0]->SetDefault(DefaultLocation.X);
+			DoubleChannels[1]->SetDefault(DefaultLocation.Y);
+			DoubleChannels[2]->SetDefault(DefaultLocation.Z);
+			DoubleChannels[3]->SetDefault(DefaultRotation.X);
+			DoubleChannels[4]->SetDefault(DefaultRotation.Y);
+			DoubleChannels[5]->SetDefault(DefaultRotation.Z);
+			DoubleChannels[6]->SetDefault(DefaultScale.X);
+			DoubleChannels[7]->SetDefault(DefaultScale.Y);
+			DoubleChannels[8]->SetDefault(DefaultScale.Z);
+
+			TArray<FVector> LocalTranslations, LocalRotations, LocalScales;
+			LocalTranslations.SetNum(KeyTimes.Num());
+			LocalRotations.SetNum(KeyTimes.Num());
+			LocalScales.SetNum(KeyTimes.Num());
+
+			for (int32 Counter = 0; Counter < KeyTimes.Num(); ++Counter)
+			{
+				FVector LocalTranslation = DefaultLocation;
+				FVector LocalScale = DefaultScale;
+				FRotator LocalRotation = DefaultRotation.Rotation();
+
+				if (Counter < BakeData.Value.Locations.Num())
+				{
+					LocalTranslation = BakeData.Value.Locations[Counter];
+				}
+				if (Counter < BakeData.Value.Rotations.Num())
+				{
+					LocalRotation = BakeData.Value.Rotations[Counter];
+				}
+				if (Counter < BakeData.Value.Scales.Num())
+				{
+					LocalScale = BakeData.Value.Scales[Counter];
+				}
+
+				FTransform LocalTransform(LocalRotation, LocalTranslation, LocalScale);
+				LocalTranslations[Counter] = LocalTransform.GetTranslation();
+				LocalRotations[Counter] = LocalTransform.GetRotation().Euler();
+				LocalScales[Counter] = LocalTransform.GetScale3D();
+			}
+
+			// Euler filter
+			for (int32 Counter = 0; Counter < LocalRotations.Num() - 1; ++Counter)
+			{
+				FMath::WindRelativeAnglesDegrees(LocalRotations[Counter].X, LocalRotations[Counter + 1].X);
+				FMath::WindRelativeAnglesDegrees(LocalRotations[Counter].Y, LocalRotations[Counter + 1].Y);
+				FMath::WindRelativeAnglesDegrees(LocalRotations[Counter].Z, LocalRotations[Counter + 1].Z);
+			}
+			if (InSettings.BakingKeySettings == EBakingKeySettings::KeysOnly)
+			{
+				const EMovieSceneKeyInterpolation KeyInterpolation = Sequencer->GetSequencerSettings()->GetKeyInterpolation();
+
+				for (int32 Counter = 0; Counter < KeyTimes.Num(); ++Counter)
+				{
+					int ChannelIndex = 0;
+					FFrameNumber KeyTime = KeyTimes[Counter];
+					{
+						TMovieSceneChannelData<FMovieSceneDoubleValue> ChannelData = DoubleChannels[ChannelIndex++]->GetData();
+						MovieSceneToolHelpers::SetOrAddKey(ChannelData, KeyTime, LocalTranslations[Counter].X);
+					}
+					{
+						TMovieSceneChannelData<FMovieSceneDoubleValue> ChannelData = DoubleChannels[ChannelIndex++]->GetData();
+						MovieSceneToolHelpers::SetOrAddKey(ChannelData, KeyTime, LocalTranslations[Counter].Y);
+					}
+					{
+						TMovieSceneChannelData<FMovieSceneDoubleValue> ChannelData = DoubleChannels[ChannelIndex++]->GetData();
+						MovieSceneToolHelpers::SetOrAddKey(ChannelData, KeyTime, LocalTranslations[Counter].Z);
+					}
+					{
+						TMovieSceneChannelData<FMovieSceneDoubleValue> ChannelData = DoubleChannels[ChannelIndex++]->GetData();
+						MovieSceneToolHelpers::SetOrAddKey(ChannelData, KeyTime, LocalRotations[Counter].X);
+					}
+					{
+						TMovieSceneChannelData<FMovieSceneDoubleValue> ChannelData = DoubleChannels[ChannelIndex++]->GetData();
+						MovieSceneToolHelpers::SetOrAddKey(ChannelData, KeyTime, LocalRotations[Counter].Y);
+
+					}
+					{
+						TMovieSceneChannelData<FMovieSceneDoubleValue> ChannelData = DoubleChannels[ChannelIndex++]->GetData();
+						MovieSceneToolHelpers::SetOrAddKey(ChannelData, KeyTime, LocalRotations[Counter].Z);
+					}
+					{
+						TMovieSceneChannelData<FMovieSceneDoubleValue> ChannelData = DoubleChannels[ChannelIndex++]->GetData();
+						MovieSceneToolHelpers::SetOrAddKey(ChannelData, KeyTime, LocalScales[Counter].X);
+					}
+					{
+						TMovieSceneChannelData<FMovieSceneDoubleValue> ChannelData = DoubleChannels[ChannelIndex++]->GetData();
+						MovieSceneToolHelpers::SetOrAddKey(ChannelData, KeyTime, LocalScales[Counter].Y);
+					}
+					{
+						TMovieSceneChannelData<FMovieSceneDoubleValue> ChannelData = DoubleChannels[ChannelIndex++]->GetData();
+						MovieSceneToolHelpers::SetOrAddKey(ChannelData, KeyTime, LocalScales[Counter].Z);
+					}
+				}
+			}
+			else
+			{
+				for (int32 Counter = 0; Counter < KeyTimes.Num(); ++Counter)
+				{
+					FFrameNumber KeyTime = KeyTimes[Counter];
+					DoubleChannels[0]->AddLinearKey(KeyTime, LocalTranslations[Counter].X);
+					DoubleChannels[1]->AddLinearKey(KeyTime, LocalTranslations[Counter].Y);
+					DoubleChannels[2]->AddLinearKey(KeyTime, LocalTranslations[Counter].Z);
+					DoubleChannels[3]->AddLinearKey(KeyTime, LocalRotations[Counter].X);
+					DoubleChannels[4]->AddLinearKey(KeyTime, LocalRotations[Counter].Y);
+					DoubleChannels[5]->AddLinearKey(KeyTime, LocalRotations[Counter].Z);
+					DoubleChannels[6]->AddLinearKey(KeyTime, LocalScales[Counter].X);
+					DoubleChannels[7]->AddLinearKey(KeyTime, LocalScales[Counter].Y);
+					DoubleChannels[8]->AddLinearKey(KeyTime, LocalScales[Counter].Z);
+				}
+				if (InSettings.bReduceKeys == true)
+				{
+					FKeyDataOptimizationParams Param;
+					Param.bAutoSetInterpolation = true;
+					Param.Tolerance = InSettings.Tolerance;
+					TRange<FFrameNumber> Range(InSettings.StartFrame, InSettings.EndFrame);
+					Param.Range = Range;
+					MovieSceneToolHelpers::OptimizeSection(Param, TransformSection);
+				}
+			}
+		}
+	}
+
+	Sequencer->NotifyMovieSceneDataChanged(EMovieSceneDataChangeType::MovieSceneStructureItemsChanged);
+	return true;
+}
+
 void FSequencerUtilities::ShowReadOnlyError()
 {
 	FNotificationInfo Info(LOCTEXT("SequenceReadOnly", "Sequence is read only."));
diff --git a/Engine/Source/Editor/Sequencer/Public/SequencerCommands.h b/Engine/Source/Editor/Sequencer/Public/SequencerCommands.h
--- a/Engine/Source/Editor/Sequencer/Public/SequencerCommands.h
+++ b/Engine/Source/Editor/Sequencer/Public/SequencerCommands.h
@@ -321,6 +321,9 @@ public:
 	/** Toggle the visibility of the stretch box. */
 	TSharedPtr< FUICommandInfo > ToggleShowStretchBox;
 
+	/** Bake transform tracks */
+	TSharedPtr< FUICommandInfo > BakeTransform;
+
 	/** Opens the director blueprint for a sequence. */
 	TSharedPtr< FUICommandInfo > OpenDirectorBlueprint;
 
diff --git a/Engine/Source/Editor/Sequencer/Public/SequencerUtilities.h b/Engine/Source/Editor/Sequencer/Public/SequencerUtilities.h
--- a/Engine/Source/Editor/Sequencer/Public/SequencerUtilities.h
+++ b/Engine/Source/Editor/Sequencer/Public/SequencerUtilities.h
@@ -23,6 +23,7 @@ class UMovieSceneSequence;
 class UMovieSceneTimeWarpGetter;
 class ISequencer;
 class FMenuBuilder;
+struct FBakingAnimationKeySettings;
 struct FMovieSceneBinding;
 struct FMovieScenePossessable;
 struct FMovieSceneSequenceID;
@@ -327,6 +328,8 @@ struct FSequencerUtilities
 	static UE_API void ReplaceBindingWithActors(TSharedRef<ISequencer> Sequencer, const TArray<AActor*>& Actors, const FMovieSceneBindingProxy& ObjectBinding);
 	static UE_API void RemoveActorsFromBinding(TSharedRef<ISequencer> Sequencer, const TArray<AActor*>& Actors, const FMovieSceneBindingProxy& ObjectBinding);
 
+	static UE_API bool BakeTransform(TSharedRef<ISequencer> Sequencer, const TArray<FMovieSceneBindingProxy>& ObjectBindings, const FBakingAnimationKeySettings& InSettings);
+
 	/** Show a read only error if the movie scene is locked */
 	static UE_API void ShowReadOnlyError();
 	
diff --git a/Engine/Source/Editor/Sequencer/Sequencer.Build.cs b/Engine/Source/Editor/Sequencer/Sequencer.Build.cs
--- a/Engine/Source/Editor/Sequencer/Sequencer.Build.cs
+++ b/Engine/Source/Editor/Sequencer/Sequencer.Build.cs
@@ -18,6 +18,7 @@ public class Sequencer : ModuleRules
 				"ApplicationCore",
 				"AssetRegistry",
 				"CinematicCamera",
+				"Constraints",
 				"ContentBrowser",
 				"Core", 
 				"CoreUObject", 
