Sequencer: Add support for snap to nearest key with tolerance in the cinematic viewport

diff --git a/Engine/Plugins/MovieScene/LevelSequenceEditor/Source/LevelSequenceEditor/Private/CinematicViewport/SCinematicTransportRange.cpp b/Engine/Plugins/MovieScene/LevelSequenceEditor/Source/LevelSequenceEditor/Private/CinematicViewport/SCinematicTransportRange.cpp
--- a/Engine/Plugins/MovieScene/LevelSequenceEditor/Source/LevelSequenceEditor/Private/CinematicViewport/SCinematicTransportRange.cpp
+++ b/Engine/Plugins/MovieScene/LevelSequenceEditor/Source/LevelSequenceEditor/Private/CinematicViewport/SCinematicTransportRange.cpp
@@ -11,6 +11,7 @@
 #include "Styles/LevelSequenceEditorStyle.h"
 #include "ISequencer.h"
 #include "TimeToPixel.h"
+#include "TimeSliderArgs.h"
 
 #define LOCTEXT_NAMESPACE "SCinematicTransportRange"
 
@@ -84,7 +85,8 @@ void SCinematicTransportRange::SetTime(const FGeometry& MyGeometry, const FPoint
 		FMovieSceneEditorData& EditorData = Sequencer->GetFocusedMovieSceneSequence()->GetMovieScene()->GetEditorData();
 		double NewTimeSeconds = EditorData.ViewStart + (EditorData.ViewEnd - EditorData.ViewStart) * Lerp;
 
-		FFrameTime ScrubTime = NewTimeSeconds * Sequencer->GetFocusedTickResolution();
+		FFrameRate TickResolution = Sequencer->GetFocusedTickResolution();
+		FFrameTime ScrubTime = NewTimeSeconds * TickResolution;
 
 		// Clamp first, snap to frame last
 		if (Sequencer->GetSequencerSettings()->ShouldKeepCursorInPlayRangeWhileScrubbing())
@@ -93,7 +95,59 @@ void SCinematicTransportRange::SetTime(const FGeometry& MyGeometry, const FPoint
 			ScrubTime = UE::MovieScene::ClampToDiscreteRange(ScrubTime, PlaybackRange);
 		}
 
-		Sequencer->SetLocalTime(ScrubTime, ESnapTimeMode::STM_All);
+		ENearestKeyOption NearestKeyOption = ENearestKeyOption::NKO_None;
+
+		TRange<double> ViewRange = Sequencer->GetFocusedMovieSceneSequence()->GetMovieScene()->GetEditorData().GetViewRange();
+
+		TRange<FFrameNumber> VisibleFrameRange(
+			(ViewRange.GetLowerBoundValue() * TickResolution).FloorToFrame(),
+			(ViewRange.GetUpperBoundValue() * TickResolution).CeilToFrame()
+		);
+
+		TArrayView<const FFrameNumber> Keys = ActiveKeyCollection.IsValid() ? ActiveKeyCollection->GetKeysInRange(VisibleFrameRange, EFindKeyType::FKT_All) : TArrayView<const FFrameNumber>();
+
+		if (Keys.Num() > 0)
+		{
+			if (Sequencer->GetSequencerSettings()->GetSnapPlayTimeToKeys() && (Sequencer->GetSequencerSettings()->GetIsSnapEnabled() || FSlateApplication::Get().GetModifierKeys().IsShiftDown()))
+			{
+				EnumAddFlags(NearestKeyOption, ENearestKeyOption::NKO_SearchKeys);
+			}
+
+			if (Sequencer->GetSequencerSettings()->GetSnapPlayTimeToSections() && (Sequencer->GetSequencerSettings()->GetIsSnapEnabled() || FSlateApplication::Get().GetModifierKeys().IsShiftDown()))
+			{
+				EnumAddFlags(NearestKeyOption, ENearestKeyOption::NKO_SearchSections);
+			}
+		}
+
+		TArray<FMovieSceneMarkedFrame> MarkedFrames = Sequencer->GetMarkedFrames();
+
+		if (MarkedFrames.Num() > 0)
+		{
+			if (Sequencer->GetSequencerSettings()->GetSnapPlayTimeToMarkers() && (Sequencer->GetSequencerSettings()->GetIsSnapEnabled() || FSlateApplication::Get().GetModifierKeys().IsShiftDown()))
+			{
+				EnumAddFlags(NearestKeyOption, ENearestKeyOption::NKO_SearchMarkers);
+			}
+		}
+
+		if (NearestKeyOption != ENearestKeyOption::NKO_None)
+		{
+			FFrameTime NearestKey = Sequencer->OnGetNearestKey(ScrubTime, NearestKeyOption);
+
+			// Anything within 3 pixel's worth of time is a duplicate as far as we're concerned
+			FTimeToPixel TimeToPixelConverter(MyGeometry, ViewRange, Sequencer->GetFocusedTickResolution());
+
+			const float ScrubPixel = TimeToPixelConverter.FrameToPixel(ScrubTime);
+			const float NearestKeyPixel = TimeToPixelConverter.FrameToPixel(NearestKey);
+
+			static float MouseTolerance = 20.f;
+
+			if (FMath::IsNearlyEqual(ScrubPixel, NearestKeyPixel, MouseTolerance))
+			{
+				ScrubTime = NearestKey;
+			}
+		}
+		
+		Sequencer->SetLocalTime(ScrubTime, ESnapTimeMode::STM_None);
 	}
 }
 
