Sequencer: Fix comment

#rnx

diff --git a/Engine/Plugins/MovieScene/LevelSequenceEditor/Source/LevelSequenceEditor/Private/CinematicViewport/SCinematicTransportRange.cpp b/Engine/Plugins/MovieScene/LevelSequenceEditor/Source/LevelSequenceEditor/Private/CinematicViewport/SCinematicTransportRange.cpp
--- a/Engine/Plugins/MovieScene/LevelSequenceEditor/Source/LevelSequenceEditor/Private/CinematicViewport/SCinematicTransportRange.cpp
+++ b/Engine/Plugins/MovieScene/LevelSequenceEditor/Source/LevelSequenceEditor/Private/CinematicViewport/SCinematicTransportRange.cpp
@@ -133,7 +133,6 @@ void SCinematicTransportRange::SetTime(const FGeometry& MyGeometry, const FPoint
 		{
 			FFrameTime NearestKey = Sequencer->OnGetNearestKey(ScrubTime, NearestKeyOption);
 
-			// Anything within 3 pixel's worth of time is a duplicate as far as we're concerned
 			FTimeToPixel TimeToPixelConverter(MyGeometry, ViewRange, Sequencer->GetFocusedTickResolution());
 
 			const float ScrubPixel = TimeToPixelConverter.FrameToPixel(ScrubTime);
