Sequencer: Add reentrant guard to AddDefaultTracksForActor

diff --git a/Engine/Plugins/MovieScene/LevelSequenceEditor/Source/LevelSequenceEditor/Private/LevelSequenceEditorToolkit.cpp b/Engine/Plugins/MovieScene/LevelSequenceEditor/Source/LevelSequenceEditor/Private/LevelSequenceEditorToolkit.cpp
--- a/Engine/Plugins/MovieScene/LevelSequenceEditor/Source/LevelSequenceEditor/Private/LevelSequenceEditorToolkit.cpp
+++ b/Engine/Plugins/MovieScene/LevelSequenceEditor/Source/LevelSequenceEditor/Private/LevelSequenceEditorToolkit.cpp
@@ -427,6 +427,14 @@ FTransform GetTransformOrigin(TSharedPtr<ISequencer> Sequencer)
 
 void FLevelSequenceEditorToolkit::AddDefaultTracksForActor(AActor& Actor, const FGuid Binding)
 {
+	static bool bReentrantUpdate = false;
+	if (bReentrantUpdate)
+	{
+		return;
+	}
+
+	TGuardValue<bool> ReentrancyGuard(bReentrantUpdate, true);
+
 	// get focused movie scene
 	UMovieSceneSequence* Sequence = Sequencer->GetFocusedMovieSceneSequence();
 
