Sequencer: Nullptr check in AddSpawnableFromClass/Instance

#rnx

diff --git a/Engine/Plugins/MovieScene/LevelSequenceEditor/Source/LevelSequenceEditor/Private/LevelSequenceEditorSubsystem.cpp b/Engine/Plugins/MovieScene/LevelSequenceEditor/Source/LevelSequenceEditor/Private/LevelSequenceEditorSubsystem.cpp
--- a/Engine/Plugins/MovieScene/LevelSequenceEditor/Source/LevelSequenceEditor/Private/LevelSequenceEditorSubsystem.cpp
+++ b/Engine/Plugins/MovieScene/LevelSequenceEditor/Source/LevelSequenceEditor/Private/LevelSequenceEditorSubsystem.cpp
@@ -997,6 +997,12 @@ FMovieSceneBindingProxy ULevelSequenceEditorSubsystem::AddSpawnableFromInstance(
 		return FMovieSceneBindingProxy();
 	}
 
+	if (!ObjectToSpawn)
+	{
+		UE_LOG(LogLevelSequenceEditor, Error, TEXT("AddSpawnableFromInstance requires a valid ObjectToSpawn"));
+		return FMovieSceneBindingProxy();
+	}
+
 	UE::Sequencer::FCreateBindingParams Params;
 	Params.BindingNameOverride = ObjectToSpawn->GetName();
 	Params.bSpawnable = true;
@@ -1026,6 +1032,12 @@ FMovieSceneBindingProxy ULevelSequenceEditorSubsystem::AddSpawnableFromClass(UMo
 		return FMovieSceneBindingProxy();
 	}
 
+	if (!ClassToSpawn)
+	{
+		UE_LOG(LogLevelSequenceEditor, Error, TEXT("AddSpawnableFromInstance requires a valid ClassToSpawn"));
+		return FMovieSceneBindingProxy();
+	}
+
 	UE::Sequencer::FCreateBindingParams Params;
 	Params.BindingNameOverride = ClassToSpawn->GetName();
 	Params.bSpawnable = true;
