Sequencer: Nullptr check in AddSpawnableFromClass/Instance

#rnx

diff --git a/Engine/Plugins/MovieScene/LevelSequenceEditor/Source/LevelSequenceEditor/Private/LevelSequenceEditorSubsystem.cpp b/Engine/Plugins/MovieScene/LevelSequenceEditor/Source/LevelSequenceEditor/Private/LevelSequenceEditorSubsystem.cpp
--- a/Engine/Plugins/MovieScene/LevelSequenceEditor/Source/LevelSequenceEditor/Private/LevelSequenceEditorSubsystem.cpp
+++ b/Engine/Plugins/MovieScene/LevelSequenceEditor/Source/LevelSequenceEditor/Private/LevelSequenceEditorSubsystem.cpp
@@ -1034,7 +1034,7 @@ FMovieSceneBindingProxy ULevelSequenceEditorSubsystem::AddSpawnableFromClass(UMo
 
 	if (!ClassToSpawn)
 	{
-		UE_LOG(LogLevelSequenceEditor, Error, TEXT("AddSpawnableFromInstance requires a valid ClassToSpawn"));
+		UE_LOG(LogLevelSequenceEditor, Error, TEXT("AddSpawnableFromClass requires a valid ClassToSpawn"));
 		return FMovieSceneBindingProxy();
 	}
 
