ShotManagement: Fixing issue where 'Add Actor Track' would not show up for CineAssembly assets.

diff --git a/Engine/Plugins/MovieScene/LevelSequenceEditor/Source/LevelSequenceEditor/Private/Misc/LevelSequenceEditorActorBinding.cpp b/Engine/Plugins/MovieScene/LevelSequenceEditor/Source/LevelSequenceEditor/Private/Misc/LevelSequenceEditorActorBinding.cpp
--- a/Engine/Plugins/MovieScene/LevelSequenceEditor/Source/LevelSequenceEditor/Private/Misc/LevelSequenceEditorActorBinding.cpp
+++ b/Engine/Plugins/MovieScene/LevelSequenceEditor/Source/LevelSequenceEditor/Private/Misc/LevelSequenceEditorActorBinding.cpp
@@ -42,7 +42,7 @@ void FLevelSequenceEditorActorBinding::BuildSequencerAddMenu(FMenuBuilder& MenuB
 
 bool FLevelSequenceEditorActorBinding::SupportsSequence(UMovieSceneSequence* InSequence) const
 {
-	return InSequence->GetClass() == ULevelSequence::StaticClass();
+	return InSequence->GetClass()->IsChildOf(ULevelSequence::StaticClass());
 }
 
 void FLevelSequenceEditorActorBinding::AddPossessActorMenuExtensions(FMenuBuilder& MenuBuilder)
