Sequencer: FindBoundObjects can return invalid/null objects

#rnx

diff --git a/Engine/Plugins/MovieScene/LevelSequenceEditor/Source/LevelSequenceEditor/Private/LevelSequenceEditorSubsystem.cpp b/Engine/Plugins/MovieScene/LevelSequenceEditor/Source/LevelSequenceEditor/Private/LevelSequenceEditorSubsystem.cpp
--- a/Engine/Plugins/MovieScene/LevelSequenceEditor/Source/LevelSequenceEditor/Private/LevelSequenceEditorSubsystem.cpp
+++ b/Engine/Plugins/MovieScene/LevelSequenceEditor/Source/LevelSequenceEditor/Private/LevelSequenceEditorSubsystem.cpp
@@ -1711,7 +1711,17 @@ void ULevelSequenceEditorSubsystem::FixActorReferences()
 		// Possessables with parents are components so ignore them.
 		if (Possessable.GetParent().IsValid() == false)
 		{
-			if (Sequencer->FindBoundObjects(Possessable.GetGuid(), Sequencer->GetFocusedTemplateID()).Num() == 0)
+			bool bAnyValidBindings = false;
+			for (TWeakObjectPtr<> WeakObject : Sequencer->FindBoundObjects(Possessable.GetGuid(), Sequencer->GetFocusedTemplateID()))
+			{
+				if (WeakObject.IsValid())
+				{
+					bAnyValidBindings = true;
+					break;
+				}
+			}
+
+			if (!bAnyValidBindings)
 			{
 				ActorsPossessablesToFix.Add(Possessable);
 			}
