ShotManagement: Fixing issue where CineAssembly tooltip would not update Favorite Rating until saving the asset.

#rnx

diff --git a/Engine/Plugins/MovieScene/LevelSequenceEditor/Source/LevelSequenceEditor/Private/Widgets/SLevelSequenceFavoriteRating.cpp b/Engine/Plugins/MovieScene/LevelSequenceEditor/Source/LevelSequenceEditor/Private/Widgets/SLevelSequenceFavoriteRating.cpp
--- a/Engine/Plugins/MovieScene/LevelSequenceEditor/Source/LevelSequenceEditor/Private/Widgets/SLevelSequenceFavoriteRating.cpp
+++ b/Engine/Plugins/MovieScene/LevelSequenceEditor/Source/LevelSequenceEditor/Private/Widgets/SLevelSequenceFavoriteRating.cpp
@@ -1,6 +1,7 @@
 // Copyright Epic Games, Inc. All Rights Reserved.
 
 #include "Widgets/SLevelSequenceFavoriteRating.h"
+#include "AssetRegistry/IAssetRegistry.h"
 #include "MetaData/MovieSceneShotMetaData.h"
 #include "ScopedTransaction.h"
 #include "Styles/LevelSequenceEditorStyle.h"
@@ -30,16 +31,20 @@ void SLevelSequenceFavoriteRating::Construct(const FArguments& InArgs)
 							.IsChecked_Lambda([&WeakLevelSequence = this->WeakLevelSequence]()->ECheckBoxState
 								{
 									bool bValue = false;
-									if (UMovieSceneShotMetaData* MetaData = WeakLevelSequence.IsValid() ? WeakLevelSequence.Pin()->FindMetaData<UMovieSceneShotMetaData>() : nullptr)
+									if (TStrongObjectPtr<ULevelSequence> StrongLevelSequence = WeakLevelSequence.Pin())
 									{
-										bValue = MetaData->GetIsNoGood().IsSet() ? MetaData->GetIsNoGood().GetValue() : false;
+										if (UMovieSceneShotMetaData* MetaData = StrongLevelSequence->FindMetaData<UMovieSceneShotMetaData>())
+										{
+											bValue = MetaData->GetIsNoGood().IsSet() ? MetaData->GetIsNoGood().GetValue() : false;
+										}
 									}
 									return bValue ? ECheckBoxState::Checked : ECheckBoxState::Unchecked;
 								})
 							.OnCheckStateChanged(FOnCheckStateChanged::CreateLambda([&WeakLevelSequence = this->WeakLevelSequence](ECheckBoxState State)
 								{
-									if (UMovieSceneShotMetaData* MetaData = WeakLevelSequence.IsValid() ? WeakLevelSequence.Pin()->FindOrAddMetaData<UMovieSceneShotMetaData>() : nullptr)
+									if (TStrongObjectPtr<ULevelSequence> StrongLevelSequence = WeakLevelSequence.Pin())
 									{
+										UMovieSceneShotMetaData* MetaData = StrongLevelSequence->FindOrAddMetaData<UMovieSceneShotMetaData>();
 										const bool bIsNoGood = State == ECheckBoxState::Checked;
 										const FText TransactionSessionName = bIsNoGood
 											? LOCTEXT("SetLevelSequenceIsNoGoodTransaction", "Set Level Sequence Is No Good")
@@ -49,6 +54,11 @@ void SLevelSequenceFavoriteRating::Construct(const FArguments& InArgs)
 										MetaData->SetFlags(RF_Transactional);
 										MetaData->Modify();
 										MetaData->SetIsNoGood(bIsNoGood);
+
+										if (IAssetRegistry* AssetRegistry = IAssetRegistry::Get())
+										{
+											AssetRegistry->AssetUpdateTags(StrongLevelSequence.Get(), EAssetRegistryTagsCaller::FullUpdate);
+										}
 									}
 								}))
 							[
@@ -69,16 +79,20 @@ void SLevelSequenceFavoriteRating::Construct(const FArguments& InArgs)
 							.IsChecked_Lambda([&WeakLevelSequence = this->WeakLevelSequence]()->ECheckBoxState
 								{
 									bool bValue = false;
-									if (UMovieSceneShotMetaData* MetaData = WeakLevelSequence.IsValid() ? WeakLevelSequence.Pin()->FindMetaData<UMovieSceneShotMetaData>() : nullptr)
+									if (TStrongObjectPtr<ULevelSequence> StrongLevelSequence = WeakLevelSequence.Pin())
 									{
-										bValue = MetaData->GetIsFlagged().IsSet() ? MetaData->GetIsFlagged().GetValue() : false;
+										if (UMovieSceneShotMetaData* MetaData = StrongLevelSequence->FindMetaData<UMovieSceneShotMetaData>())
+										{
+											bValue = MetaData->GetIsFlagged().IsSet() ? MetaData->GetIsFlagged().GetValue() : false;
+										}
 									}
 									return bValue ? ECheckBoxState::Checked : ECheckBoxState::Unchecked;
 								})
 							.OnCheckStateChanged(FOnCheckStateChanged::CreateLambda([&WeakLevelSequence = this->WeakLevelSequence](ECheckBoxState State)
 								{
-									if (UMovieSceneShotMetaData* MetaData = WeakLevelSequence.IsValid() ? WeakLevelSequence.Pin()->FindOrAddMetaData<UMovieSceneShotMetaData>() : nullptr)
+									if (TStrongObjectPtr<ULevelSequence> StrongLevelSequence = WeakLevelSequence.Pin())
 									{
+										UMovieSceneShotMetaData* MetaData = StrongLevelSequence->FindOrAddMetaData<UMovieSceneShotMetaData>();
 										const bool bIsFlagged = State == ECheckBoxState::Checked;
 										const FText TransactionSessionName = bIsFlagged 
 											? LOCTEXT("SetLevelSequenceIsFlaggedTransaction", "Set Level Sequence Is Flagged")
@@ -88,6 +102,11 @@ void SLevelSequenceFavoriteRating::Construct(const FArguments& InArgs)
 										MetaData->SetFlags(RF_Transactional);
 										MetaData->Modify();
 										MetaData->SetIsFlagged(bIsFlagged);
+
+										if (IAssetRegistry* AssetRegistry = IAssetRegistry::Get())
+										{
+											AssetRegistry->AssetUpdateTags(StrongLevelSequence.Get(), EAssetRegistryTagsCaller::FullUpdate);
+										}
 									}
 								}))
 							[
@@ -106,8 +125,9 @@ void SLevelSequenceFavoriteRating::Construct(const FArguments& InArgs)
 							.MaxScore(3)
 							.OnScoreChanged_Lambda([&WeakLevelSequence = this->WeakLevelSequence](int32 NewScore)
 								{
-									if (UMovieSceneShotMetaData* MetaData = WeakLevelSequence.IsValid() ? WeakLevelSequence.Pin()->FindOrAddMetaData<UMovieSceneShotMetaData>() : nullptr)
+									if (TStrongObjectPtr<ULevelSequence> StrongLevelSequence = WeakLevelSequence.Pin())
 									{
+										UMovieSceneShotMetaData* MetaData = StrongLevelSequence->FindOrAddMetaData<UMovieSceneShotMetaData>();
 										const FScopedTransaction Transaction(LOCTEXT("SetLevelSequenceFavoriteRatingTransaction", "Set Level Sequence Favorite Rating"));
 										MetaData->SetFlags(RF_Transactional);
 										MetaData->Modify();
@@ -119,12 +139,24 @@ void SLevelSequenceFavoriteRating::Construct(const FArguments& InArgs)
 										{
 											MetaData->ClearFavoriteRating();
 										}
+
+										if (IAssetRegistry* AssetRegistry = IAssetRegistry::Get())
+										{
+											AssetRegistry->AssetUpdateTags(StrongLevelSequence.Get(), EAssetRegistryTagsCaller::FullUpdate);
+										}
 									}
 								})
 							.Score_Lambda([&WeakLevelSequence = this->WeakLevelSequence]() -> int32
 								{
-									UMovieSceneShotMetaData* MetaData = WeakLevelSequence.IsValid() ? WeakLevelSequence.Pin()->FindMetaData<UMovieSceneShotMetaData>() : nullptr;
-									return MetaData && MetaData->GetFavoriteRating().IsSet() ? MetaData->GetFavoriteRating().GetValue() : 0;
+									if (TStrongObjectPtr<ULevelSequence> StrongLevelSequence = WeakLevelSequence.Pin())
+									{
+										if (UMovieSceneShotMetaData* MetaData = StrongLevelSequence->FindMetaData<UMovieSceneShotMetaData>();
+											MetaData != nullptr && MetaData->GetFavoriteRating().IsSet())
+										{
+											return MetaData->GetFavoriteRating().GetValue();
+										}
+									}
+									return 0;
 								})
 
 					]
