UMG: Disable Dynamic Possession menu if it's not supported

Added a new SupportsDirectorBlueprint to FMovieSceneSequenceEditor and derived classes implement CanCreateDirectorBlueprint and check for class permissions

diff --git a/Engine/Plugins/Experimental/DaySequence/Source/DaySequenceEditor/DaySequenceEditor.Build.cs b/Engine/Plugins/Experimental/DaySequence/Source/DaySequenceEditor/DaySequenceEditor.Build.cs
--- a/Engine/Plugins/Experimental/DaySequence/Source/DaySequenceEditor/DaySequenceEditor.Build.cs
+++ b/Engine/Plugins/Experimental/DaySequence/Source/DaySequenceEditor/DaySequenceEditor.Build.cs
@@ -20,6 +20,7 @@ public class DaySequenceEditor : ModuleRules
 			{
 				"ApplicationCore",
 				"AssetDefinition",
+				"ClassViewer",
 				"Projects",
 				"InputCore",
 				"EditorFramework",
diff --git a/Engine/Plugins/Experimental/DaySequence/Source/DaySequenceEditor/Private/MovieSceneSequenceEditor_DaySequence.h b/Engine/Plugins/Experimental/DaySequence/Source/DaySequenceEditor/Private/MovieSceneSequenceEditor_DaySequence.h
--- a/Engine/Plugins/Experimental/DaySequence/Source/DaySequenceEditor/Private/MovieSceneSequenceEditor_DaySequence.h
+++ b/Engine/Plugins/Experimental/DaySequence/Source/DaySequenceEditor/Private/MovieSceneSequenceEditor_DaySequence.h
@@ -4,10 +4,13 @@
 
 #include "Engine/Blueprint.h"
 #include "MovieSceneSequenceEditor.h"
+#include "ClassViewerFilter.h"
+#include "ClassViewerModule.h"
 #include "DaySequenceDirector.h"
 #include "DaySequence.h"
 #include "Engine/BlueprintGeneratedClass.h"
 #include "Kismet2/KismetEditorUtilities.h"
+#include "Modules/ModuleManager.h"
 
 struct FMovieSceneSequenceEditor_DaySequence : FMovieSceneSequenceEditor
 {
@@ -16,6 +19,22 @@ struct FMovieSceneSequenceEditor_DaySequence : FMovieSceneSequenceEditor
 		return true;
 	}
 
+	virtual bool CanCreateDirectorBlueprint(UMovieSceneSequence* Sequence) const override
+	{
+#if WITH_EDITOR
+		FClassViewerModule& ClassViewerModule = FModuleManager::LoadModuleChecked<FClassViewerModule>("ClassViewer");
+		const TSharedPtr<IClassViewerFilter>& GlobalClassFilter = ClassViewerModule.GetGlobalClassViewerFilter();
+		TSharedRef<FClassViewerFilterFuncs> ClassFilterFuncs = ClassViewerModule.CreateFilterFuncs();
+		FClassViewerInitializationOptions ClassViewerOptions = {};
+
+		if (GlobalClassFilter.IsValid())
+		{
+			return GlobalClassFilter->IsClassAllowed(ClassViewerOptions, UDaySequenceDirector::StaticClass(), ClassFilterFuncs);
+		}
+#endif
+		return true;
+	}
+
 	virtual UBlueprint* GetBlueprintForSequence(UMovieSceneSequence* InSequence) const override
 	{
 		UDaySequence* DaySequence = CastChecked<UDaySequence>(InSequence);
diff --git a/Engine/Plugins/MovieScene/ActorSequence/Source/ActorSequenceEditor/ActorSequenceEditor.Build.cs b/Engine/Plugins/MovieScene/ActorSequence/Source/ActorSequenceEditor/ActorSequenceEditor.Build.cs
--- a/Engine/Plugins/MovieScene/ActorSequence/Source/ActorSequenceEditor/ActorSequenceEditor.Build.cs
+++ b/Engine/Plugins/MovieScene/ActorSequence/Source/ActorSequenceEditor/ActorSequenceEditor.Build.cs
@@ -10,6 +10,7 @@ public class ActorSequenceEditor : ModuleRules
 			new string[] {
 				"ActorSequence",
 				"BlueprintGraph",
+				"ClassViewer",
 				"Core",
 				"CoreUObject",
 				"EditorFramework",
diff --git a/Engine/Plugins/MovieScene/ActorSequence/Source/ActorSequenceEditor/Private/MovieSceneSequenceEditor_ActorSequence.h b/Engine/Plugins/MovieScene/ActorSequence/Source/ActorSequenceEditor/Private/MovieSceneSequenceEditor_ActorSequence.h
--- a/Engine/Plugins/MovieScene/ActorSequence/Source/ActorSequenceEditor/Private/MovieSceneSequenceEditor_ActorSequence.h
+++ b/Engine/Plugins/MovieScene/ActorSequence/Source/ActorSequenceEditor/Private/MovieSceneSequenceEditor_ActorSequence.h
@@ -6,7 +6,10 @@
 #include "GameFramework/Actor.h"
 #include "MovieSceneSequenceEditor.h"
 #include "ActorSequence.h"
+#include "ClassViewerFilter.h"
+#include "ClassViewerModule.h"
 #include "Engine/LevelScriptBlueprint.h"
+#include "Modules/ModuleManager.h"
 
 struct FMovieSceneSequenceEditor_ActorSequence : FMovieSceneSequenceEditor
 {
@@ -15,6 +18,22 @@ struct FMovieSceneSequenceEditor_ActorSequence : FMovieSceneSequenceEditor
 		return true;
 	}
 
+	virtual bool CanCreateDirectorBlueprint(UMovieSceneSequence* Sequence) const override
+	{
+#if WITH_EDITOR
+		FClassViewerModule& ClassViewerModule = FModuleManager::LoadModuleChecked<FClassViewerModule>("ClassViewer");
+		const TSharedPtr<IClassViewerFilter>& GlobalClassFilter = ClassViewerModule.GetGlobalClassViewerFilter();
+		TSharedRef<FClassViewerFilterFuncs> ClassFilterFuncs = ClassViewerModule.CreateFilterFuncs();
+		FClassViewerInitializationOptions ClassViewerOptions = {};
+
+		if (GlobalClassFilter.IsValid())
+		{
+			return GlobalClassFilter->IsClassAllowed(ClassViewerOptions, UBlueprint::StaticClass(), ClassFilterFuncs);
+		}
+#endif
+		return true;
+	}
+
 	virtual UBlueprint* GetBlueprintForSequence(UMovieSceneSequence* InSequence) const override
 	{
 		UActorSequence* ActorSequence = CastChecked<UActorSequence>(InSequence);
diff --git a/Engine/Plugins/MovieScene/LevelSequenceEditor/Source/LevelSequenceEditor/Private/Misc/MovieSceneSequenceEditor_LevelSequence.h b/Engine/Plugins/MovieScene/LevelSequenceEditor/Source/LevelSequenceEditor/Private/Misc/MovieSceneSequenceEditor_LevelSequence.h
--- a/Engine/Plugins/MovieScene/LevelSequenceEditor/Source/LevelSequenceEditor/Private/Misc/MovieSceneSequenceEditor_LevelSequence.h
+++ b/Engine/Plugins/MovieScene/LevelSequenceEditor/Source/LevelSequenceEditor/Private/Misc/MovieSceneSequenceEditor_LevelSequence.h
@@ -6,8 +6,11 @@
 #include "LevelSequence.h"
 #include "MovieScene.h"
 #include "MovieSceneSequenceEditor.h"
+#include "ClassViewerFilter.h"
+#include "ClassViewerModule.h"
 #include "LevelSequenceDirector.h"
 #include "Kismet2/KismetEditorUtilities.h"
+#include "Modules/ModuleManager.h"
 
 struct FMovieSceneSequenceEditor_LevelSequence : FMovieSceneSequenceEditor
 {
@@ -16,6 +19,22 @@ struct FMovieSceneSequenceEditor_LevelSequence : FMovieSceneSequenceEditor
 		return true;
 	}
 
+	virtual bool CanCreateDirectorBlueprint(UMovieSceneSequence* Sequence) const override
+	{
+#if WITH_EDITOR
+		FClassViewerModule& ClassViewerModule = FModuleManager::LoadModuleChecked<FClassViewerModule>("ClassViewer");
+		const TSharedPtr<IClassViewerFilter>& GlobalClassFilter = ClassViewerModule.GetGlobalClassViewerFilter();
+		TSharedRef<FClassViewerFilterFuncs> ClassFilterFuncs = ClassViewerModule.CreateFilterFuncs();
+		FClassViewerInitializationOptions ClassViewerOptions = {};
+
+		if (GlobalClassFilter.IsValid())
+		{
+			return GlobalClassFilter->IsClassAllowed(ClassViewerOptions, ULevelSequenceDirector::StaticClass(), ClassFilterFuncs);
+		}
+#endif
+		return true;
+	}
+
 	virtual UBlueprint* GetBlueprintForSequence(UMovieSceneSequence* InSequence) const override
 	{
 		ULevelSequence* LevelSequence = CastChecked<ULevelSequence>(InSequence);
@@ -24,13 +43,6 @@ struct FMovieSceneSequenceEditor_LevelSequence : FMovieSceneSequenceEditor
 
 	virtual UBlueprint* CreateBlueprintForSequence(UMovieSceneSequence* InSequence) const override
 	{
-#if WITH_EDITOR
-		if (!UMovieScene::IsTrackClassAllowed(ULevelSequenceDirector::StaticClass()))
-		{
-			return nullptr;
-		}
-#endif
-
 		UBlueprint* Blueprint = GetBlueprintForSequence(InSequence);
 		if (!ensureMsgf(!Blueprint, TEXT("Should not call CreateBlueprintForSequence when one already exists")))
 		{
diff --git a/Engine/Plugins/VirtualProduction/Avalanche/Source/AvalancheSequencer/AvalancheSequencer.Build.cs b/Engine/Plugins/VirtualProduction/Avalanche/Source/AvalancheSequencer/AvalancheSequencer.Build.cs
--- a/Engine/Plugins/VirtualProduction/Avalanche/Source/AvalancheSequencer/AvalancheSequencer.Build.cs
+++ b/Engine/Plugins/VirtualProduction/Avalanche/Source/AvalancheSequencer/AvalancheSequencer.Build.cs
@@ -33,6 +33,7 @@ public class AvalancheSequencer : ModuleRules
 				"AvalancheTransition",
 				"AvalancheTransitionEditor",
 				"BlueprintGraph",
+				"ClassViewer",
 				"CustomDetailsView",
 				"DeveloperSettings",
 				"EditorSubsystem",
diff --git a/Engine/Plugins/VirtualProduction/Avalanche/Source/AvalancheSequencer/Private/AvaSequenceEditor.cpp b/Engine/Plugins/VirtualProduction/Avalanche/Source/AvalancheSequencer/Private/AvaSequenceEditor.cpp
--- a/Engine/Plugins/VirtualProduction/Avalanche/Source/AvalancheSequencer/Private/AvaSequenceEditor.cpp
+++ b/Engine/Plugins/VirtualProduction/Avalanche/Source/AvalancheSequencer/Private/AvaSequenceEditor.cpp
@@ -6,13 +6,32 @@
 #include "Director/AvaSequenceDirectorBlueprint.h"
 #include "Director/AvaSequenceDirectorGeneratedClass.h"
 #include "IAvaSequenceProvider.h"
+#include "ClassViewerFilter.h"
+#include "ClassViewerModule.h"
 #include "Kismet2/KismetEditorUtilities.h"
+#include "Modules/ModuleManager.h"
 
 bool FAvaSequenceEditor::CanCreateEvents(UMovieSceneSequence* InSequence) const
 {
 	return true;
 }
 
+bool FAvaSequenceEditor::CanCreateDirectorBlueprint(UMovieSceneSequence* Sequence) const
+{
+#if WITH_EDITOR
+	FClassViewerModule& ClassViewerModule = FModuleManager::LoadModuleChecked<FClassViewerModule>("ClassViewer");
+	const TSharedPtr<IClassViewerFilter>& GlobalClassFilter = ClassViewerModule.GetGlobalClassViewerFilter();
+	TSharedRef<FClassViewerFilterFuncs> ClassFilterFuncs = ClassViewerModule.CreateFilterFuncs();
+	FClassViewerInitializationOptions ClassViewerOptions = {};
+
+	if (GlobalClassFilter.IsValid())
+	{
+		return GlobalClassFilter->IsClassAllowed(ClassViewerOptions, UBlueprint::StaticClass(), ClassFilterFuncs);
+	}
+#endif
+	return true;
+}
+
 UBlueprint* FAvaSequenceEditor::GetBlueprintForSequence(UMovieSceneSequence* InSequence) const
 {
 	UAvaSequence* const Sequence = CastChecked<UAvaSequence>(InSequence);
diff --git a/Engine/Plugins/VirtualProduction/Avalanche/Source/AvalancheSequencer/Private/AvaSequenceEditor.h b/Engine/Plugins/VirtualProduction/Avalanche/Source/AvalancheSequencer/Private/AvaSequenceEditor.h
--- a/Engine/Plugins/VirtualProduction/Avalanche/Source/AvalancheSequencer/Private/AvaSequenceEditor.h
+++ b/Engine/Plugins/VirtualProduction/Avalanche/Source/AvalancheSequencer/Private/AvaSequenceEditor.h
@@ -9,6 +9,7 @@ struct FAvaSequenceEditor : FMovieSceneSequenceEditor
 private:
 	//~ Begin FMovieSceneSequenceEditor
 	virtual bool CanCreateEvents(UMovieSceneSequence* InSequence) const override;
+	virtual bool CanCreateDirectorBlueprint(UMovieSceneSequence* InSequence) const override;
 	virtual UBlueprint* GetBlueprintForSequence(UMovieSceneSequence* InSequence) const override;
 	virtual UBlueprint* CreateBlueprintForSequence(UMovieSceneSequence* InSequence) const override;
 	//~ End FMovieSceneSequenceEditor
diff --git a/Engine/Source/Editor/Sequencer/Private/MovieSceneSequenceEditor.cpp b/Engine/Source/Editor/Sequencer/Private/MovieSceneSequenceEditor.cpp
--- a/Engine/Source/Editor/Sequencer/Private/MovieSceneSequenceEditor.cpp
+++ b/Engine/Source/Editor/Sequencer/Private/MovieSceneSequenceEditor.cpp
@@ -17,6 +17,11 @@ bool FMovieSceneSequenceEditor::SupportsEvents(UMovieSceneSequence* InSequence)
 	return CanCreateEvents(InSequence);
 }
 
+bool FMovieSceneSequenceEditor::SupportsDirectorBlueprint(UMovieSceneSequence* InSequence) const
+{
+	return CanCreateDirectorBlueprint(InSequence);
+}
+
 UBlueprint* FMovieSceneSequenceEditor::FindDirectorBlueprint(UMovieSceneSequence* Sequence) const
 {
 	return GetBlueprintForSequence(Sequence);
@@ -24,6 +29,11 @@ UBlueprint* FMovieSceneSequenceEditor::FindDirectorBlueprint(UMovieSceneSequence
 
 UBlueprint* FMovieSceneSequenceEditor::GetOrCreateDirectorBlueprint(UMovieSceneSequence* Sequence) const
 {
+	if (!SupportsDirectorBlueprint(Sequence))
+	{
+		return nullptr;
+	}
+
 	UBlueprint* Blueprint = GetBlueprintForSequence(Sequence);
 	if (!Blueprint)
 	{
diff --git a/Engine/Source/Editor/Sequencer/Public/MovieSceneSequenceEditor.h b/Engine/Source/Editor/Sequencer/Public/MovieSceneSequenceEditor.h
--- a/Engine/Source/Editor/Sequencer/Public/MovieSceneSequenceEditor.h
+++ b/Engine/Source/Editor/Sequencer/Public/MovieSceneSequenceEditor.h
@@ -23,8 +23,6 @@ struct FMovieSceneSequenceEditor
 	 */
 	static UE_API FMovieSceneSequenceEditor* Find(UMovieSceneSequence* InSequence);
 
-
-
 	/**
 	 * Check whether the specified sequence supports events
 	 *
@@ -32,6 +30,12 @@ struct FMovieSceneSequenceEditor
 	 */
 	UE_API bool SupportsEvents(UMovieSceneSequence* InSequence) const;
 
+	/**
+	 * Check whether the specified sequence supports director blueprints
+	 *
+	 * @param Sequence        The sequence to test
+	 */
+	UE_API bool SupportsDirectorBlueprint(UMovieSceneSequence* Sequence) const;
 
 	/**
 	 * Access the director blueprint for the specified sequence
@@ -52,6 +56,7 @@ struct FMovieSceneSequenceEditor
 private:
 
 	virtual bool CanCreateEvents(UMovieSceneSequence* InSequence) const { return false; }
+	virtual bool CanCreateDirectorBlueprint(UMovieSceneSequence* Sequence) const { return false; }
 	virtual UBlueprint* GetBlueprintForSequence(UMovieSceneSequence* InSequence) const { return nullptr; }
 	virtual UBlueprint* CreateBlueprintForSequence(UMovieSceneSequence* InSequence) const { return nullptr; }
 };
diff --git a/Engine/Source/Editor/UMGEditor/Private/Animation/MovieSceneSequenceEditor_WidgetAnimation.h b/Engine/Source/Editor/UMGEditor/Private/Animation/MovieSceneSequenceEditor_WidgetAnimation.h
--- a/Engine/Source/Editor/UMGEditor/Private/Animation/MovieSceneSequenceEditor_WidgetAnimation.h
+++ b/Engine/Source/Editor/UMGEditor/Private/Animation/MovieSceneSequenceEditor_WidgetAnimation.h
@@ -5,8 +5,11 @@
 #include "MovieSceneSequence.h"
 #include "MovieSceneSequenceEditor.h"
 #include "Tracks/MovieSceneEventTrack.h"
+#include "ClassViewerFilter.h"
+#include "ClassViewerModule.h"
 #include "Engine/Blueprint.h"
 #include "K2Node_FunctionEntry.h"
+#include "Modules/ModuleManager.h"
 
 struct FMovieSceneSequenceEditor_WidgetAnimation : FMovieSceneSequenceEditor
 {
@@ -15,6 +18,22 @@ struct FMovieSceneSequenceEditor_WidgetAnimation : FMovieSceneSequenceEditor
 		return true;
 	}
 
+	virtual bool CanCreateDirectorBlueprint(UMovieSceneSequence* Sequence) const override
+	{
+#if WITH_EDITOR
+		FClassViewerModule& ClassViewerModule = FModuleManager::LoadModuleChecked<FClassViewerModule>("ClassViewer");
+		const TSharedPtr<IClassViewerFilter>& GlobalClassFilter = ClassViewerModule.GetGlobalClassViewerFilter();
+		TSharedRef<FClassViewerFilterFuncs> ClassFilterFuncs = ClassViewerModule.CreateFilterFuncs();
+		FClassViewerInitializationOptions ClassViewerOptions = {};
+
+		if (GlobalClassFilter.IsValid())
+		{
+			return GlobalClassFilter->IsClassAllowed(ClassViewerOptions, UBlueprint::StaticClass(), ClassFilterFuncs);
+		}
+#endif
+		return true;
+	}
+
 	virtual UBlueprint* GetBlueprintForSequence(UMovieSceneSequence* InSequence) const override
 	{
 		return InSequence->GetTypedOuter<UBlueprint>();
diff --git a/Engine/Source/Editor/UMGEditor/Private/WidgetBlueprintEditor.cpp b/Engine/Source/Editor/UMGEditor/Private/WidgetBlueprintEditor.cpp
--- a/Engine/Source/Editor/UMGEditor/Private/WidgetBlueprintEditor.cpp
+++ b/Engine/Source/Editor/UMGEditor/Private/WidgetBlueprintEditor.cpp
@@ -90,6 +90,7 @@
 #include "UnrealEdGlobals.h"
 #include "GraphEditorActions.h"
 #include "MovieSceneDynamicBindingCustomization.h"
+#include "MovieSceneSequenceEditor.h"
 #include "UIComponentWidgetBlueprintExtension.h"
 
 #include "Extensions/UIComponent.h"
@@ -2241,10 +2242,13 @@ void FWidgetBlueprintEditor::OnBuildCustomContextMenuForGuid(FMenuBuilder& MenuB
 				FExecuteAction::CreateRaw(this, &FWidgetBlueprintEditor::RemoveMissingWidgetsFromTrack, ObjectBinding)
 			);
 
-			MenuBuilder.AddSubMenu(
-				LOCTEXT("DynamicPossession", "Dynamic Possession"),
-				LOCTEXT("DynamicPossessionToolTip", "Specify a Blueprint method that will find a compatible widget for this binding"),
-				FNewMenuDelegate::CreateRaw(this, &FWidgetBlueprintEditor::AddDynamicPossessionMenu, ObjectBinding));
+			if (CanAddDynamicPossessionMenu())
+			{
+				MenuBuilder.AddSubMenu(
+					LOCTEXT("DynamicPossession", "Dynamic Possession"),
+					LOCTEXT("DynamicPossessionToolTip", "Specify a Blueprint method that will find a compatible widget for this binding"),
+					FNewMenuDelegate::CreateRaw(this, &FWidgetBlueprintEditor::AddDynamicPossessionMenu, ObjectBinding));
+			}
 		}
 	}
 }
@@ -2521,6 +2525,21 @@ void FWidgetBlueprintEditor::ReplaceTrackWithWidgets(TArray<FWidgetReference> Wi
 	SyncSequencersMovieSceneData();
 }
 
+bool FWidgetBlueprintEditor::CanAddDynamicPossessionMenu()
+{
+	TSharedPtr<ISequencer>& ActiveSequencer = GetSequencer();
+
+	UWidgetAnimation* WidgetAnimation = Cast<UWidgetAnimation>(ActiveSequencer->GetFocusedMovieSceneSequence());
+	if (!WidgetAnimation)
+	{
+		return false;
+	}
+
+	ISequencerModule& SequencerModule = FModuleManager::LoadModuleChecked<ISequencerModule>("Sequencer");
+	FMovieSceneSequenceEditor* SequenceEditor = SequencerModule.FindSequenceEditor(WidgetAnimation->GetClass());
+	return SequenceEditor && SequenceEditor->SupportsDirectorBlueprint(WidgetAnimation);
+}
+
 void FWidgetBlueprintEditor::AddDynamicPossessionMenu(FMenuBuilder& MenuBuilder, FGuid ObjectId)
 {
 	using namespace UE::Sequencer;
diff --git a/Engine/Source/Editor/UMGEditor/Public/WidgetBlueprintEditor.h b/Engine/Source/Editor/UMGEditor/Public/WidgetBlueprintEditor.h
--- a/Engine/Source/Editor/UMGEditor/Public/WidgetBlueprintEditor.h
+++ b/Engine/Source/Editor/UMGEditor/Public/WidgetBlueprintEditor.h
@@ -409,6 +409,7 @@ private:
 	UE_API void ReplaceTrackWithWidgets(const TArray<FWidgetReference> Widgets, FGuid ObjectId);
 
 	/** Dynamic binding */
+	UE_API bool CanAddDynamicPossessionMenu();
 	UE_API void AddDynamicPossessionMenu(FMenuBuilder& MenuBuilder, FGuid ObjectId);
 
 	/** Add an animation track for the supplied slot to the current animation. */
