[Film Overlays] Fix an issue where the "tint" value was reset every time the configuration menu was opened.

[RN] Editor Minor

diff --git a/Engine/Plugins/MovieScene/LevelSequenceEditor/Source/LevelSequenceEditor/Private/CinematicViewport/FilmOverlays.cpp b/Engine/Plugins/MovieScene/LevelSequenceEditor/Source/LevelSequenceEditor/Private/CinematicViewport/FilmOverlays.cpp
--- a/Engine/Plugins/MovieScene/LevelSequenceEditor/Source/LevelSequenceEditor/Private/CinematicViewport/FilmOverlays.cpp
+++ b/Engine/Plugins/MovieScene/LevelSequenceEditor/Source/LevelSequenceEditor/Private/CinematicViewport/FilmOverlays.cpp
@@ -565,6 +565,11 @@ void SFilmOverlayOptions::Construct(const FArguments& InArgs, TSharedPtr<SFilmOv
 	OverlayWidget = InFilmOverlay;
 	if (TSharedPtr<SFilmOverlay> OverlayWidgetPinned = OverlayWidget.Pin())
 	{
+		if (IFilmOverlay* OldPrimaryOverlay = GetPrimaryFilmOverlay())
+		{
+			PrimaryColorTint = OldPrimaryOverlay->GetTint();
+		}
+	
 		OverlayWidgetPinned->SetFilmOverlays(TAttribute<TArray<IFilmOverlay*>>::CreateSP(this, &SFilmOverlayOptions::GetActiveFilmOverlays));
 
 		// CurrentPrimaryOverlay (and therefore Thumbnail) resets for non button version of this widget every time we create it, so let's keep it up to date
