Sequencer: Fix BakeTransformWithSettings not using the converted Settings

diff --git a/Engine/Plugins/MovieScene/LevelSequenceEditor/Source/LevelSequenceEditor/Private/LevelSequenceEditorSubsystem.cpp b/Engine/Plugins/MovieScene/LevelSequenceEditor/Source/LevelSequenceEditor/Private/LevelSequenceEditorSubsystem.cpp
--- a/Engine/Plugins/MovieScene/LevelSequenceEditor/Source/LevelSequenceEditor/Private/LevelSequenceEditorSubsystem.cpp
+++ b/Engine/Plugins/MovieScene/LevelSequenceEditor/Source/LevelSequenceEditor/Private/LevelSequenceEditorSubsystem.cpp
@@ -1663,7 +1663,7 @@ bool ULevelSequenceEditorSubsystem::BakeTransformWithSettings(const TArray<FMovi
 		Settings.EndFrame = ConvertFrameTime(Settings.EndFrame, DisplayRate, TickResolution).GetFrame();
 	}
 
-	return FSequencerUtilities::BakeTransform(Sequencer.ToSharedRef(), ObjectBindings, InSettings);
+	return FSequencerUtilities::BakeTransform(Sequencer.ToSharedRef(), ObjectBindings, Settings);
 }
 
 void ULevelSequenceEditorSubsystem::FixActorReferences()
