Sequencer: Follow up to 47603852

#rnx

diff --git a/Engine/Plugins/MovieScene/LevelSequenceEditor/Source/LevelSequenceEditor/Private/LevelSequenceEditorModule.cpp b/Engine/Plugins/MovieScene/LevelSequenceEditor/Source/LevelSequenceEditor/Private/LevelSequenceEditorModule.cpp
--- a/Engine/Plugins/MovieScene/LevelSequenceEditor/Source/LevelSequenceEditor/Private/LevelSequenceEditorModule.cpp
+++ b/Engine/Plugins/MovieScene/LevelSequenceEditor/Source/LevelSequenceEditor/Private/LevelSequenceEditorModule.cpp
@@ -383,8 +383,11 @@ protected:
 	{
 		if (UObjectInitialized())
 		{
-			ISequencerModule& SequencerModule = FModuleManager::LoadModuleChecked<ISequencerModule>("Sequencer");
-			SequencerModule.GetSequencerCustomizationManager()->UnregisterInstancedSequencerCustomization(ULevelSequence::StaticClass());
+			ISequencerModule* SequencerModule = FModuleManager::GetModulePtr<ISequencerModule>("Sequencer");
+			if (SequencerModule)
+			{
+				SequencerModule->GetSequencerCustomizationManager()->UnregisterInstancedSequencerCustomization(ULevelSequence::StaticClass());
+			}
 		}
 	}
 
