Sequencer: Cleaned up a few of the Level Sequence Editor actions menu extensions.
The Import/Export FBX options are now added to the customization builder that is set up by the LevelSequenceCustomization for a specific sequencer, instead of adding them to the Sequencer module's global extension manager, which fixes an issue where the same extender could be added more than once if more than one Sequencer was opened.
The Snap Sections / Sync Sections extensions have been moved out of the LevelSequenceEditorSubsystem and into the LevelSequenceCustomization, and the implementation moved into SequencerUtilities. This makes it possible to provide a specific ISequencer to the implementation of those menu extensions, rather than using the "Active" Sequencer defined by the subsystem.

#rnx
#Synced-CL 47206406

diff --git a/Engine/Plugins/MovieScene/LevelSequenceEditor/Source/LevelSequenceEditor/Private/LevelSequenceEditorSubsystem.cpp b/Engine/Plugins/MovieScene/LevelSequenceEditor/Source/LevelSequenceEditor/Private/LevelSequenceEditorSubsystem.cpp
--- a/Engine/Plugins/MovieScene/LevelSequenceEditor/Source/LevelSequenceEditor/Private/LevelSequenceEditorSubsystem.cpp
+++ b/Engine/Plugins/MovieScene/LevelSequenceEditor/Source/LevelSequenceEditor/Private/LevelSequenceEditorSubsystem.cpp
@@ -627,35 +627,9 @@ void ULevelSequenceEditorSubsystem::Initialize(FSubsystemCollectionBase& Collect
 		return SelectedActors.Num() > 0;
 	};
 
-	auto AreMovieSceneSectionsSelected = [this](const int32 MinSections = 1) {
-		const TSharedPtr<ISequencer> Sequencer = GetActiveSequencer();
-		if (!Sequencer)
-		{
-			return false;
-		}
-
-		TArray<UMovieSceneSection*> SelectedSections;
-		Sequencer->GetSelectedSections(SelectedSections);
-		return (SelectedSections.Num() >= MinSections);
-	};
-
 	/* Commands for this subsystem */
 	CommandList = MakeShareable(new FUICommandList);
 
-	CommandList->MapAction(FLevelSequenceEditorCommands::Get().SnapSectionsToTimelineUsingSourceTimecode,
-		FExecuteAction::CreateUObject(this, &ULevelSequenceEditorSubsystem::SnapSectionsToTimelineUsingSourceTimecodeInternal),
-		FCanExecuteAction::CreateLambda(AreMovieSceneSectionsSelected)
-	);
-
-	CommandList->MapAction(FLevelSequenceEditorCommands::Get().SyncSectionsUsingSourceTimecode,
-		FExecuteAction::CreateUObject(this, &ULevelSequenceEditorSubsystem::SyncSectionsUsingSourceTimecodeInternal),
-		FCanExecuteAction::CreateLambda(AreMovieSceneSectionsSelected, 2)
-	);
-
-	CommandList->MapAction(FLevelSequenceEditorCommands::Get().FixActorReferences,
-		FExecuteAction::CreateUObject(this, &ULevelSequenceEditorSubsystem::FixActorReferences)
-	);
-
 	CommandList->MapAction(FLevelSequenceEditorCommands::Get().AddActorsToBinding,
 		FExecuteAction::CreateUObject(this, &ULevelSequenceEditorSubsystem::AddActorsToBindingInternal),
 		FCanExecuteAction::CreateLambda(AreActorsSelected));
@@ -675,33 +649,6 @@ void ULevelSequenceEditorSubsystem::Initialize(FSubsystemCollectionBase& Collect
 		FExecuteAction::CreateUObject(this, &ULevelSequenceEditorSubsystem::RemoveInvalidBindingsInternal));
 
 	/* Menu extenders */
-	TransformMenuExtender = MakeShareable(new FExtender);
-	TransformMenuExtender->AddMenuExtension("Transform", EExtensionHook::After, CommandList, FMenuExtensionDelegate::CreateLambda([this](FMenuBuilder& MenuBuilder) {
-		// Only add menu entries where the focused sequence is a ULevelSequence
-		if (!GetActiveSequencer())
-		{
-			return;
-		}
-		
-		MenuBuilder.AddMenuEntry(FLevelSequenceEditorCommands::Get().SnapSectionsToTimelineUsingSourceTimecode);
-		MenuBuilder.AddMenuEntry(FLevelSequenceEditorCommands::Get().SyncSectionsUsingSourceTimecode);
-		}));
-
-	SequencerModule.GetActionsMenuExtensibilityManager()->AddExtender(TransformMenuExtender);
-
-	FixActorReferencesMenuExtender = MakeShareable(new FExtender);
-	FixActorReferencesMenuExtender->AddMenuExtension("Bindings", EExtensionHook::First, CommandList, FMenuExtensionDelegate::CreateLambda([this](FMenuBuilder& MenuBuilder) {
-		// Only add menu entries where the focused sequence is a ULevelSequence
-		if (!GetActiveSequencer())
-		{
-			return;
-		}
-		
-		MenuBuilder.AddMenuEntry(FLevelSequenceEditorCommands::Get().FixActorReferences);
-		}));
-
-	SequencerModule.GetActionsMenuExtensibilityManager()->AddExtender(FixActorReferencesMenuExtender);
-
 	AssignActorMenuExtender = MakeShareable(new FExtender);
 	AssignActorMenuExtender->AddMenuExtension("Possessable", EExtensionHook::First, CommandList, FMenuExtensionDelegate::CreateLambda([this](FMenuBuilder& MenuBuilder) {
 
@@ -1452,195 +1399,19 @@ bool ULevelSequenceEditorSubsystem::PasteBindings(const FString& InTextToImport,
 	return true;
 }
 
-void ULevelSequenceEditorSubsystem::SnapSectionsToTimelineUsingSourceTimecodeInternal()
-{
-	const TSharedPtr<ISequencer> Sequencer = GetActiveSequencer();
-	if (!Sequencer)
-	{
-		return;
-	}
-
-	const UMovieScene* FocusedMovieScene = Sequencer->GetFocusedMovieSceneSequence()->GetMovieScene();
-	if (!FocusedMovieScene)
-	{
-		return;
-	}
-
-	if (FocusedMovieScene->IsReadOnly())
-	{
-		FSequencerUtilities::ShowReadOnlyError();
-		return;
-	}
-
-	TArray<UMovieSceneSection*> Sections;
-	Sequencer->GetSelectedSections(Sections);
-	if (Sections.IsEmpty())
-	{
-		return;
-	}
-
-	SnapSectionsToTimelineUsingSourceTimecode(Sections);
-}
-
 void ULevelSequenceEditorSubsystem::SnapSectionsToTimelineUsingSourceTimecode(const TArray<UMovieSceneSection*>& Sections)
 {
-	const TSharedPtr<ISequencer> Sequencer = GetActiveSequencer();
-	if (!Sequencer)
-	{
-		return;
-	}
-
-	const UMovieScene* FocusedMovieScene = Sequencer->GetFocusedMovieSceneSequence()->GetMovieScene();
-	if (!FocusedMovieScene)
-	{
-		return;
-	}
-
-	if (FocusedMovieScene->IsReadOnly())
-	{
-		FSequencerUtilities::ShowReadOnlyError();
-		return;
-	}
-
-	FScopedTransaction SnapSectionsToTimelineUsingSourceTimecodeTransaction(LOCTEXT("SnapSectionsToTimelineUsingSourceTimecode_Transaction", "Snap Sections to Timeline using Source Timecode"));
-	bool bAnythingChanged = false;
-
-	const FFrameRate TickResolution = Sequencer->GetFocusedTickResolution();
-	const FFrameRate DisplayRate = Sequencer->GetFocusedDisplayRate();
-
-	for (UMovieSceneSection* Section : Sections)
-	{
-		if (!Section || !Section->HasStartFrame())
-		{
-			continue;
-		}
-
-		const FTimecode SectionSourceTimecode = Section->TimecodeSource.Timecode;
-		if (SectionSourceTimecode == FTimecode())
-		{
-			// Do not move sections with default values for source timecode.
-			continue;
-		}
-
-		const FFrameNumber SectionSourceStartFrameNumber = SectionSourceTimecode.ToFrameNumber(DisplayRate);
-
-		// Account for any trimming at the start of the section when computing the
-		// target frame number to move this section to.
-		const FFrameNumber SectionOffsetFrames = Section->GetOffsetTime().Get(FFrameTime()).FloorToFrame();
-		const FFrameNumber TargetFrameNumber = SectionSourceStartFrameNumber + SectionOffsetFrames;
-
-		const FFrameNumber SectionCurrentStartFrameNumber = Section->GetInclusiveStartFrame();
-
-		const FFrameNumber Delta = -(SectionCurrentStartFrameNumber - ConvertFrameTime(TargetFrameNumber, DisplayRate, TickResolution).GetFrame().Value);
-
-		Section->MoveSection(Delta);
-
-		bAnythingChanged |= (Delta.Value != 0);
-	}
-
-	if (bAnythingChanged)
+	if (TSharedPtr<ISequencer> Sequencer = GetActiveSequencer())
 	{
-		Sequencer->NotifyMovieSceneDataChanged(EMovieSceneDataChangeType::TrackValueChanged);
+		FSequencerUtilities::SnapSectionsToTimelineUsingSourceTimecode(Sequencer.ToSharedRef(), Sections);
 	}
 }
 
-void ULevelSequenceEditorSubsystem::SyncSectionsUsingSourceTimecodeInternal()
-{
-	const TSharedPtr<ISequencer> Sequencer = GetActiveSequencer();
-	if (!Sequencer)
-	{
-		return;
-	}
-
-	const UMovieScene* FocusedMovieScene = Sequencer->GetFocusedMovieSceneSequence()->GetMovieScene();
-	if (!FocusedMovieScene)
-	{
-		return;
-	}
-
-	if (FocusedMovieScene->IsReadOnly())
-	{
-		FSequencerUtilities::ShowReadOnlyError();
-		return;
-	}
-
-	TArray<UMovieSceneSection*> Sections;
-	Sequencer->GetSelectedSections(Sections);
-	if (Sections.Num() < 2)
-	{
-		return;
-	}
-
-	SyncSectionsUsingSourceTimecode(Sections);
-}
-
 void ULevelSequenceEditorSubsystem::SyncSectionsUsingSourceTimecode(const TArray<UMovieSceneSection*>& Sections)
 {
-	const TSharedPtr<ISequencer> Sequencer = GetActiveSequencer();
-	if (!Sequencer)
-	{
-		return;
-	}
-
-	const UMovieScene* FocusedMovieScene = Sequencer->GetFocusedMovieSceneSequence()->GetMovieScene();
-	if (!FocusedMovieScene)
-	{
-		return;
-	}
-
-	if (FocusedMovieScene->IsReadOnly())
-	{
-		FSequencerUtilities::ShowReadOnlyError();
-		return;
-	}
-
-	// Pull out all of the valid sections that have a start frame and verify
-	// we have at least two sections to sync.
-	TArray<UMovieSceneSection*> SectionsToSync;
-	for (UMovieSceneSection* Section : Sections)
-	{
-		if (Section && Section->HasStartFrame())
-		{
-			SectionsToSync.Add(Section);
-		}
-	}
-
-	if (SectionsToSync.Num() < 2)
+	if (TSharedPtr<ISequencer> Sequencer = GetActiveSequencer())
 	{
-		return;
-	}
-
-	FScopedTransaction SyncSectionsUsingSourceTimecodeTransaction(LOCTEXT("SyncSectionsUsingSourceTimecode_Transaction", "Sync Sections Using Source Timecode"));
-	bool bAnythingChanged = false;
-
-	const FFrameRate TickResolution = Sequencer->GetFocusedTickResolution();
-	const FFrameRate DisplayRate = Sequencer->GetFocusedDisplayRate();
-
-	const UMovieSceneSection* FirstSection = SectionsToSync[0];
-	const FFrameNumber FirstSectionSourceTimecode = FirstSection->TimecodeSource.Timecode.ToFrameNumber(DisplayRate);
-
-	const FFrameNumber FirstSectionCurrentStartFrame = FirstSection->GetInclusiveStartFrame();
-	const FFrameNumber FirstSectionOffsetFrames = FirstSection->GetOffsetTime().Get(FFrameTime()).FloorToFrame();
-	SectionsToSync.RemoveAt(0);
-
-	for (UMovieSceneSection* Section : SectionsToSync)
-	{
-		const FFrameNumber SectionSourceTimecode = Section->TimecodeSource.Timecode.ToFrameNumber(DisplayRate);
-		const FFrameNumber SectionCurrentStartFrame = Section->GetInclusiveStartFrame();
-		const FFrameNumber SectionOffsetFrames = Section->GetOffsetTime().Get(FFrameTime()).FloorToFrame();
-
-		const FFrameNumber TimecodeDelta = ConvertFrameTime(SectionSourceTimecode - FirstSectionSourceTimecode, DisplayRate, TickResolution).GetFrame().Value;
-		const FFrameNumber CurrentDelta = (SectionCurrentStartFrame - SectionOffsetFrames) - (FirstSectionCurrentStartFrame - FirstSectionOffsetFrames);
-		const FFrameNumber Delta = -CurrentDelta + TimecodeDelta;
-
-		Section->MoveSection(Delta);
-
-		bAnythingChanged |= (Delta.Value != 0);
-	}
-
-	if (bAnythingChanged)
-	{
-		Sequencer->NotifyMovieSceneDataChanged(EMovieSceneDataChangeType::TrackValueChanged);
+		FSequencerUtilities::SyncSectionsUsingSourceTimecode(Sequencer.ToSharedRef(), Sections);
 	}
 }
 
@@ -1668,96 +1439,10 @@ bool ULevelSequenceEditorSubsystem::BakeTransformWithSettings(const TArray<FMovi
 
 void ULevelSequenceEditorSubsystem::FixActorReferences()
 {
-	TSharedPtr<ISequencer> Sequencer = GetActiveSequencer();
-	if (Sequencer == nullptr)
+	if (TSharedPtr<ISequencer> Sequencer = GetActiveSequencer())
 	{
-		return;
+		FSequencerUtilities::FixActorReferences(Sequencer.ToSharedRef());
 	}
-
-	UWorld* PlaybackContext = Sequencer->GetPlaybackContext()->GetWorld();
-	if (!PlaybackContext)
-	{
-		return;
-	}
-
-	UMovieScene* FocusedMovieScene = Sequencer->GetFocusedMovieSceneSequence()->GetMovieScene();
-	if (!FocusedMovieScene)
-	{
-		return;
-	}
-
-	if (FocusedMovieScene->IsReadOnly())
-	{
-		FSequencerUtilities::ShowReadOnlyError();
-		return;
-	}
-
-	FScopedTransaction FixActorReferencesTransaction(LOCTEXT("FixActorReferences", "Fix Actor References"));
-
-	TMap<FString, AActor*> ActorNameToActorMap;
-
-	for (TActorIterator<AActor> ActorItr(PlaybackContext); ActorItr; ++ActorItr)
-	{
-		// Same as with the Object Iterator, access the subclass instance with the * or -> operators.
-		AActor* Actor = *ActorItr;
-		ActorNameToActorMap.Add(Actor->GetActorLabel(), Actor);
-	}
-
-	// Cache the possessables to fix up first since the bindings will change as the fix ups happen.
-	TArray<FMovieScenePossessable> ActorsPossessablesToFix;
-	for (int32 i = 0; i < FocusedMovieScene->GetPossessableCount(); i++)
-	{
-		FMovieScenePossessable& Possessable = FocusedMovieScene->GetPossessable(i);
-		// Possessables with parents are components so ignore them.
-		if (Possessable.GetParent().IsValid() == false)
-		{
-			bool bAnyValidBindings = false;
-			for (TWeakObjectPtr<> WeakObject : Sequencer->FindBoundObjects(Possessable.GetGuid(), Sequencer->GetFocusedTemplateID()))
-			{
-				if (WeakObject.IsValid())
-				{
-					bAnyValidBindings = true;
-					break;
-				}
-			}
-
-			if (!bAnyValidBindings)
-			{
-				ActorsPossessablesToFix.Add(Possessable);
-			}
-		}
-	}
-
-	// For the possessables to fix, look up the actors by name and reassign them if found.
-	TMap<FGuid, FGuid> OldGuidToNewGuidMap;
-	for (const FMovieScenePossessable& ActorPossessableToFix : ActorsPossessablesToFix)
-	{
-		AActor* ActorPtr = ActorNameToActorMap.FindRef(ActorPossessableToFix.GetName());
-		if (ActorPtr != nullptr)
-		{
-			FGuid OldGuid = ActorPossessableToFix.GetGuid();
-
-			// The actor might have an existing guid while the possessable with the same name might not. 
-			// In that case, make sure we also replace the existing guid with the new guid 
-			FGuid ExistingGuid = Sequencer->FindObjectId(*ActorPtr, Sequencer->GetFocusedTemplateID());
-
-			FGuid NewGuid = FSequencerUtilities::AssignActor(Sequencer.ToSharedRef(), ActorPtr, ActorPossessableToFix.GetGuid());
-
-			OldGuidToNewGuidMap.Add(OldGuid, NewGuid);
-
-			if (ExistingGuid.IsValid())
-			{
-				OldGuidToNewGuidMap.Add(ExistingGuid, NewGuid);
-			}
-		}
-	}
-
-	for (TPair<FGuid, FGuid> GuidPair : OldGuidToNewGuidMap)
-	{
-		FSequencerUtilities::UpdateBindingIDs(Sequencer.ToSharedRef(), GuidPair.Key, GuidPair.Value);
-	}
-
-	Sequencer->NotifyMovieSceneDataChanged(EMovieSceneDataChangeType::MovieSceneStructureItemsChanged);
 }
 
 void ULevelSequenceEditorSubsystem::AddActorsToBindingInternal()
diff --git a/Engine/Plugins/MovieScene/LevelSequenceEditor/Source/LevelSequenceEditor/Private/Misc/LevelSequenceCustomization.cpp b/Engine/Plugins/MovieScene/LevelSequenceEditor/Source/LevelSequenceEditor/Private/Misc/LevelSequenceCustomization.cpp
--- a/Engine/Plugins/MovieScene/LevelSequenceEditor/Source/LevelSequenceEditor/Private/Misc/LevelSequenceCustomization.cpp
+++ b/Engine/Plugins/MovieScene/LevelSequenceEditor/Source/LevelSequenceEditor/Private/Misc/LevelSequenceCustomization.cpp
@@ -39,29 +39,10 @@ void FLevelSequenceCustomization::RegisterSequencerCustomization(FSequencerCusto
 {
 	WeakSequencer = Builder.GetSequencer().AsShared();
 
-	const FLevelSequenceEditorCommands& Commands = FLevelSequenceEditorCommands::Get();
-
-	// Build the extender for the actions menu.
-	ActionsMenuCommandList = MakeShared<FUICommandList>().ToSharedPtr();
-	ActionsMenuCommandList->MapAction(
-		Commands.ImportFBX,
-		FExecuteAction::CreateRaw( this, &FLevelSequenceCustomization::ImportFBX ),
-		FCanExecuteAction::CreateLambda( [] { return true; } ) );
-	ActionsMenuCommandList->MapAction(
-		Commands.ExportFBX,
-		FExecuteAction::CreateRaw( this, &FLevelSequenceCustomization::ExportFBX ),
-		FCanExecuteAction::CreateLambda( [] { return true; } ) );
-
-	ActionsMenuExtender = MakeShared<FExtender>();
-	ActionsMenuExtender->AddMenuExtension(
-			"SequenceOptions", EExtensionHook::First, nullptr,
-			FMenuExtensionDelegate::CreateRaw(this, &FLevelSequenceCustomization::ExtendActionsMenu));
+	BindActionMenuCommands();
 
-	ISequencerModule& SequencerModule = FModuleManager::LoadModuleChecked<ISequencerModule>("Sequencer");
-	SequencerModule.GetActionsMenuExtensibilityManager()->AddExtender(ActionsMenuExtender);
-
-	// Add a customization callback for the object binding context menu.
 	FSequencerCustomizationInfo Customization;
+	Customization.ActionsMenuExtender = CreateActionsMenuExtender();
 	Customization.OnBuildObjectBindingContextMenu = FOnGetSequencerMenuExtender::CreateRaw(this, &FLevelSequenceCustomization::CreateObjectBindingContextMenuExtender);
 	Customization.OnBuildSidebarMenu = FOnGetSequencerMenuExtender::CreateRaw(this, &FLevelSequenceCustomization::CreateObjectBindingSidebarMenuExtender);
 	Builder.AddCustomization(Customization);
@@ -69,8 +50,7 @@ void FLevelSequenceCustomization::RegisterSequencerCustomization(FSequencerCusto
 	for(TUniquePtr<ISequencerCustomization>& ExternalCustomization : AdditionalCustomizations)
 	{
 		ExternalCustomization->RegisterSequencerCustomization(Builder);
-	}
-	
+	}	
 }
 
 void FLevelSequenceCustomization::UnregisterSequencerCustomization()
@@ -80,23 +60,71 @@ void FLevelSequenceCustomization::UnregisterSequencerCustomization()
 		ExternalCustomization->UnregisterSequencerCustomization();
 	}
 
-	ISequencerModule& SequencerModule = FModuleManager::LoadModuleChecked<ISequencerModule>("Sequencer");
-	SequencerModule.GetActionsMenuExtensibilityManager()->RemoveExtender(ActionsMenuExtender);
-
 	ActionsMenuCommandList = nullptr;
 	WeakSequencer = nullptr;
 }
 
-void FLevelSequenceCustomization::ExtendActionsMenu(FMenuBuilder& MenuBuilder)
+void FLevelSequenceCustomization::BindActionMenuCommands()
 {
-	MenuBuilder.PushCommandList(ActionsMenuCommandList.ToSharedRef());
-	{
-		const FLevelSequenceEditorCommands& Commands = FLevelSequenceEditorCommands::Get();
-		
-		MenuBuilder.AddMenuEntry(Commands.ImportFBX);
-		MenuBuilder.AddMenuEntry(Commands.ExportFBX);
-	}
-	MenuBuilder.PopCommandList();
+	const FLevelSequenceEditorCommands& Commands = FLevelSequenceEditorCommands::Get();
+
+	ActionsMenuCommandList = MakeShared<FUICommandList>().ToSharedPtr();
+
+	ActionsMenuCommandList->MapAction(
+		Commands.ImportFBX,
+		FExecuteAction::CreateRaw(this, &FLevelSequenceCustomization::ImportFBX),
+		FCanExecuteAction::CreateLambda([] { return true; }));
+
+	ActionsMenuCommandList->MapAction(
+		Commands.ExportFBX,
+		FExecuteAction::CreateRaw(this, &FLevelSequenceCustomization::ExportFBX),
+		FCanExecuteAction::CreateLambda([] { return true; }));
+
+	constexpr int32 MinSectionsRequiredForSnapping = 1;
+	ActionsMenuCommandList->MapAction(
+		Commands.SnapSectionsToTimelineUsingSourceTimecode,
+		FExecuteAction::CreateRaw(this, &FLevelSequenceCustomization::SnapSectionsToTimelineUsingSourceTimecode),
+		FCanExecuteAction::CreateRaw(this, &FLevelSequenceCustomization::DoesMovieSceneHaveEnoughSectionsSelected, MinSectionsRequiredForSnapping));
+
+	constexpr int32 MinSectionsRequiredForSyncing = 2;
+	ActionsMenuCommandList->MapAction(
+		Commands.SyncSectionsUsingSourceTimecode,
+		FExecuteAction::CreateRaw(this, &FLevelSequenceCustomization::SyncSectionsUsingSourceTimecode),
+		FCanExecuteAction::CreateRaw(this, &FLevelSequenceCustomization::DoesMovieSceneHaveEnoughSectionsSelected, MinSectionsRequiredForSyncing));
+
+	ActionsMenuCommandList->MapAction(
+		Commands.FixActorReferences,
+		FExecuteAction::CreateRaw(this, &FLevelSequenceCustomization::FixActorReferences));
+}
+
+TSharedPtr<FExtender> FLevelSequenceCustomization::CreateActionsMenuExtender()
+{
+	TSharedPtr<FExtender> ActionsMenuExtender = MakeShared<FExtender>();
+
+	ActionsMenuExtender->AddMenuExtension(
+		"SequenceOptions", EExtensionHook::First, ActionsMenuCommandList,
+		FMenuExtensionDelegate::CreateLambda([](FMenuBuilder& MenuBuilder)
+			{
+				MenuBuilder.AddMenuEntry(FLevelSequenceEditorCommands::Get().ImportFBX);
+				MenuBuilder.AddMenuEntry(FLevelSequenceEditorCommands::Get().ExportFBX);
+			}));
+
+	ActionsMenuExtender->AddMenuExtension(
+		"Transform", EExtensionHook::After, ActionsMenuCommandList,
+		FMenuExtensionDelegate::CreateLambda([](FMenuBuilder& MenuBuilder)
+			{
+				MenuBuilder.AddMenuEntry(FLevelSequenceEditorCommands::Get().SnapSectionsToTimelineUsingSourceTimecode);
+				MenuBuilder.AddMenuEntry(FLevelSequenceEditorCommands::Get().SyncSectionsUsingSourceTimecode);
+			}));
+
+	ActionsMenuExtender->AddMenuExtension(
+		"Bindings", EExtensionHook::First, ActionsMenuCommandList,
+		FMenuExtensionDelegate::CreateLambda([this](FMenuBuilder& MenuBuilder)
+			{
+				MenuBuilder.AddMenuEntry(FLevelSequenceEditorCommands::Get().FixActorReferences);
+			}));
+
+	return ActionsMenuExtender;
 }
 
 void FLevelSequenceCustomization::ImportFBX()
@@ -111,6 +139,48 @@ void FLevelSequenceCustomization::ExportFBX()
 	Interop.ExportFBX();
 }
 
+void FLevelSequenceCustomization::SnapSectionsToTimelineUsingSourceTimecode()
+{
+	if (const TSharedPtr<ISequencer> Sequencer = WeakSequencer.Pin())
+	{
+		TArray<UMovieSceneSection*> Sections;
+		Sequencer->GetSelectedSections(Sections);
+		FSequencerUtilities::SnapSectionsToTimelineUsingSourceTimecode(Sequencer.ToSharedRef(), Sections);
+	}
+}
+
+void FLevelSequenceCustomization::SyncSectionsUsingSourceTimecode()
+{
+	if (const TSharedPtr<ISequencer> Sequencer = WeakSequencer.Pin())
+	{
+		TArray<UMovieSceneSection*> Sections;
+		Sequencer->GetSelectedSections(Sections);
+		FSequencerUtilities::SyncSectionsUsingSourceTimecode(Sequencer.ToSharedRef(), Sections);
+	}
+}
+
+bool FLevelSequenceCustomization::DoesMovieSceneHaveEnoughSectionsSelected(int32 MinSectionCount)
+{
+	const TSharedPtr<ISequencer> Sequencer = WeakSequencer.Pin();
+	if (!Sequencer)
+	{
+		return false;
+	}
+
+	TArray<UMovieSceneSection*> SelectedSections;
+	Sequencer->GetSelectedSections(SelectedSections);
+
+	return (SelectedSections.Num() >= MinSectionCount);
+}
+
+void FLevelSequenceCustomization::FixActorReferences()
+{
+	if (const TSharedPtr<ISequencer> Sequencer = WeakSequencer.Pin())
+	{
+		FSequencerUtilities::FixActorReferences(Sequencer.ToSharedRef());
+	}
+}
+
 TSharedPtr<FExtender> FLevelSequenceCustomization::CreateObjectBindingContextMenuExtender(FViewModelPtr InViewModel)
 {
 	TSharedRef<FExtender> Extender = MakeShared<FExtender>();
diff --git a/Engine/Plugins/MovieScene/LevelSequenceEditor/Source/LevelSequenceEditor/Private/Misc/LevelSequenceCustomization.h b/Engine/Plugins/MovieScene/LevelSequenceEditor/Source/LevelSequenceEditor/Private/Misc/LevelSequenceCustomization.h
--- a/Engine/Plugins/MovieScene/LevelSequenceEditor/Source/LevelSequenceEditor/Private/Misc/LevelSequenceCustomization.h
+++ b/Engine/Plugins/MovieScene/LevelSequenceEditor/Source/LevelSequenceEditor/Private/Misc/LevelSequenceCustomization.h
@@ -29,9 +29,14 @@ protected:
 private:
 
 	// Actions menu extensions
-	void ExtendActionsMenu(FMenuBuilder& MenuBuilder);
+	void BindActionMenuCommands();
+	TSharedPtr<FExtender> CreateActionsMenuExtender();
 	void ImportFBX();
 	void ExportFBX();
+	void SnapSectionsToTimelineUsingSourceTimecode();
+	void SyncSectionsUsingSourceTimecode();
+	bool DoesMovieSceneHaveEnoughSectionsSelected(int32 MinSectionCount);
+	void FixActorReferences();
 
 	// Object binding context menu extensions
 	TSharedPtr<FExtender> CreateObjectBindingContextMenuExtender(FViewModelPtr InViewModel);
@@ -46,7 +51,6 @@ private:
 	TWeakPtr<ISequencer> WeakSequencer;
 
 	TSharedPtr<FUICommandList> ActionsMenuCommandList;
-	TSharedPtr<FExtender> ActionsMenuExtender;
 
 	TArray<TUniquePtr<ISequencerCustomization>>	AdditionalCustomizations;
 };
diff --git a/Engine/Plugins/MovieScene/LevelSequenceEditor/Source/LevelSequenceEditor/Public/LevelSequenceEditorSubsystem.h b/Engine/Plugins/MovieScene/LevelSequenceEditor/Source/LevelSequenceEditor/Public/LevelSequenceEditorSubsystem.h
--- a/Engine/Plugins/MovieScene/LevelSequenceEditor/Source/LevelSequenceEditor/Public/LevelSequenceEditorSubsystem.h
+++ b/Engine/Plugins/MovieScene/LevelSequenceEditor/Source/LevelSequenceEditor/Public/LevelSequenceEditorSubsystem.h
@@ -309,8 +309,6 @@ private:
 
 	UE_API TSharedPtr<ISequencer> GetActiveSequencer();
 	
-	UE_API void SnapSectionsToTimelineUsingSourceTimecodeInternal();
-	UE_API void SyncSectionsUsingSourceTimecodeInternal();
 	UE_API void AddActorsToBindingInternal();
 	UE_API void ReplaceBindingWithActorsInternal();
 	UE_API void RemoveActorsFromBindingInternal();
diff --git a/Engine/Source/Editor/Sequencer/Private/SequencerUtilities.cpp b/Engine/Source/Editor/Sequencer/Private/SequencerUtilities.cpp
--- a/Engine/Source/Editor/Sequencer/Private/SequencerUtilities.cpp
+++ b/Engine/Source/Editor/Sequencer/Private/SequencerUtilities.cpp
@@ -5055,6 +5055,220 @@ UMovieSceneSequence* FSequencerUtilities::GetMovieSceneSequence(TSharedPtr<ISequ
 	return InSequencer->GetRootMovieSceneSequence();
 }
 
+void FSequencerUtilities::SnapSectionsToTimelineUsingSourceTimecode(TSharedRef<ISequencer> Sequencer, const TArray<UMovieSceneSection*>& Sections)
+{
+	const UMovieScene* FocusedMovieScene = Sequencer->GetFocusedMovieSceneSequence()->GetMovieScene();
+	if (!FocusedMovieScene)
+	{
+		return;
+	}
+
+	if (FocusedMovieScene->IsReadOnly())
+	{
+		ShowReadOnlyError();
+		return;
+	}
+
+	FScopedTransaction SnapSectionsToTimelineUsingSourceTimecodeTransaction(LOCTEXT("SnapSectionsToTimelineUsingSourceTimecode_Transaction", "Snap Sections to Timeline using Source Timecode"));
+	bool bAnythingChanged = false;
+
+	const FFrameRate TickResolution = Sequencer->GetFocusedTickResolution();
+	const FFrameRate DisplayRate = Sequencer->GetFocusedDisplayRate();
+
+	for (UMovieSceneSection* Section : Sections)
+	{
+		if (!Section || !Section->HasStartFrame())
+		{
+			continue;
+		}
+
+		const FTimecode SectionSourceTimecode = Section->TimecodeSource.Timecode;
+		if (SectionSourceTimecode == FTimecode())
+		{
+			// Do not move sections with default values for source timecode.
+			continue;
+		}
+
+		const FFrameNumber SectionSourceStartFrameNumber = SectionSourceTimecode.ToFrameNumber(DisplayRate);
+
+		// Account for any trimming at the start of the section when computing the
+		// target frame number to move this section to.
+		const FFrameNumber SectionOffsetFrames = Section->GetOffsetTime().Get(FFrameTime()).FloorToFrame();
+		const FFrameNumber TargetFrameNumber = SectionSourceStartFrameNumber + SectionOffsetFrames;
+
+		const FFrameNumber SectionCurrentStartFrameNumber = Section->GetInclusiveStartFrame();
+
+		const FFrameNumber Delta = -(SectionCurrentStartFrameNumber - ConvertFrameTime(TargetFrameNumber, DisplayRate, TickResolution).GetFrame().Value);
+
+		Section->MoveSection(Delta);
+
+		bAnythingChanged |= (Delta.Value != 0);
+	}
+
+	if (bAnythingChanged)
+	{
+		Sequencer->NotifyMovieSceneDataChanged(EMovieSceneDataChangeType::TrackValueChanged);
+	}
+}
+
+void FSequencerUtilities::SyncSectionsUsingSourceTimecode(TSharedRef<ISequencer> Sequencer, const TArray<UMovieSceneSection*>& Sections)
+{
+	const UMovieScene* FocusedMovieScene = Sequencer->GetFocusedMovieSceneSequence()->GetMovieScene();
+	if (!FocusedMovieScene)
+	{
+		return;
+	}
+
+	if (FocusedMovieScene->IsReadOnly())
+	{
+		ShowReadOnlyError();
+		return;
+	}
+
+	// Pull out all of the valid sections that have a start frame and verify
+	// we have at least two sections to sync.
+	TArray<UMovieSceneSection*> SectionsToSync;
+	for (UMovieSceneSection* Section : Sections)
+	{
+		if (Section && Section->HasStartFrame())
+		{
+			SectionsToSync.Add(Section);
+		}
+	}
+
+	if (SectionsToSync.Num() < 2)
+	{
+		return;
+	}
+
+	FScopedTransaction SyncSectionsUsingSourceTimecodeTransaction(LOCTEXT("SyncSectionsUsingSourceTimecode_Transaction", "Sync Sections Using Source Timecode"));
+	bool bAnythingChanged = false;
+
+	const FFrameRate TickResolution = Sequencer->GetFocusedTickResolution();
+	const FFrameRate DisplayRate = Sequencer->GetFocusedDisplayRate();
+
+	const UMovieSceneSection* FirstSection = SectionsToSync[0];
+	const FFrameNumber FirstSectionSourceTimecode = FirstSection->TimecodeSource.Timecode.ToFrameNumber(DisplayRate);
+
+	const FFrameNumber FirstSectionCurrentStartFrame = FirstSection->GetInclusiveStartFrame();
+	const FFrameNumber FirstSectionOffsetFrames = FirstSection->GetOffsetTime().Get(FFrameTime()).FloorToFrame();
+	SectionsToSync.RemoveAt(0);
+
+	for (UMovieSceneSection* Section : SectionsToSync)
+	{
+		const FFrameNumber SectionSourceTimecode = Section->TimecodeSource.Timecode.ToFrameNumber(DisplayRate);
+		const FFrameNumber SectionCurrentStartFrame = Section->GetInclusiveStartFrame();
+		const FFrameNumber SectionOffsetFrames = Section->GetOffsetTime().Get(FFrameTime()).FloorToFrame();
+
+		const FFrameNumber TimecodeDelta = ConvertFrameTime(SectionSourceTimecode - FirstSectionSourceTimecode, DisplayRate, TickResolution).GetFrame().Value;
+		const FFrameNumber CurrentDelta = (SectionCurrentStartFrame - SectionOffsetFrames) - (FirstSectionCurrentStartFrame - FirstSectionOffsetFrames);
+		const FFrameNumber Delta = -CurrentDelta + TimecodeDelta;
+
+		Section->MoveSection(Delta);
+
+		bAnythingChanged |= (Delta.Value != 0);
+	}
+
+	if (bAnythingChanged)
+	{
+		Sequencer->NotifyMovieSceneDataChanged(EMovieSceneDataChangeType::TrackValueChanged);
+	}
+}
+
+void FSequencerUtilities::FixActorReferences(TSharedRef<ISequencer> Sequencer)
+{
+	UWorld* PlaybackContext = Sequencer->GetPlaybackContext()->GetWorld();
+	if (!PlaybackContext)
+	{
+		return;
+	}
+
+	UMovieSceneSequence* FocusedSequence = Sequencer->GetFocusedMovieSceneSequence();
+	if (!FocusedSequence)
+	{
+		return;
+	}
+
+	UMovieScene* FocusedMovieScene = FocusedSequence->GetMovieScene();
+	if (!FocusedMovieScene)
+	{
+		return;
+	}
+
+	if (FocusedMovieScene->IsReadOnly())
+	{
+		ShowReadOnlyError();
+		return;
+	}
+
+	FScopedTransaction FixActorReferencesTransaction(LOCTEXT("FixActorReferences", "Fix Actor References"));
+
+	TMap<FString, AActor*> ActorNameToActorMap;
+
+	for (TActorIterator<AActor> ActorItr(PlaybackContext); ActorItr; ++ActorItr)
+	{
+		// Same as with the Object Iterator, access the subclass instance with the * or -> operators.
+		AActor* Actor = *ActorItr;
+		ActorNameToActorMap.Add(Actor->GetActorLabel(), Actor);
+	}
+
+	// Cache the possessables to fix up first since the bindings will change as the fix ups happen.
+	TArray<FMovieScenePossessable> ActorsPossessablesToFix;
+	for (int32 i = 0; i < FocusedMovieScene->GetPossessableCount(); i++)
+	{
+		FMovieScenePossessable& Possessable = FocusedMovieScene->GetPossessable(i);
+		// Possessables with parents are components so ignore them.
+		if (Possessable.GetParent().IsValid() == false)
+		{
+			bool bAnyValidBindings = false;
+			for (TWeakObjectPtr<> WeakObject : Sequencer->FindBoundObjects(Possessable.GetGuid(), Sequencer->GetFocusedTemplateID()))
+			{
+				if (WeakObject.IsValid())
+				{
+					bAnyValidBindings = true;
+					break;
+				}
+			}
+
+			if (!bAnyValidBindings)
+			{
+				ActorsPossessablesToFix.Add(Possessable);
+			}
+		}
+	}
+
+	// For the possessables to fix, look up the actors by name and reassign them if found.
+	TMap<FGuid, FGuid> OldGuidToNewGuidMap;
+	for (const FMovieScenePossessable& ActorPossessableToFix : ActorsPossessablesToFix)
+	{
+		AActor* ActorPtr = ActorNameToActorMap.FindRef(ActorPossessableToFix.GetName());
+		if (ActorPtr != nullptr)
+		{
+			FGuid OldGuid = ActorPossessableToFix.GetGuid();
+
+			// The actor might have an existing guid while the possessable with the same name might not. 
+			// In that case, make sure we also replace the existing guid with the new guid 
+			FGuid ExistingGuid = Sequencer->FindObjectId(*ActorPtr, Sequencer->GetFocusedTemplateID());
+
+			FGuid NewGuid = AssignActor(Sequencer, ActorPtr, ActorPossessableToFix.GetGuid());
+
+			OldGuidToNewGuidMap.Add(OldGuid, NewGuid);
+
+			if (ExistingGuid.IsValid())
+			{
+				OldGuidToNewGuidMap.Add(ExistingGuid, NewGuid);
+			}
+		}
+	}
+
+	for (TPair<FGuid, FGuid> GuidPair : OldGuidToNewGuidMap)
+	{
+		UpdateBindingIDs(Sequencer, GuidPair.Key, GuidPair.Value);
+	}
+
+	Sequencer->NotifyMovieSceneDataChanged(EMovieSceneDataChangeType::MovieSceneStructureItemsChanged);
+}
+
 void FOpenSequencerWatcher::DoStartup(TFunction<void()> StartupComplete)
 {
 	auto RegisterWatcher = [this, StartupCompleteFunc = MoveTemp(StartupComplete)]()
diff --git a/Engine/Source/Editor/Sequencer/Public/SequencerUtilities.h b/Engine/Source/Editor/Sequencer/Public/SequencerUtilities.h
--- a/Engine/Source/Editor/Sequencer/Public/SequencerUtilities.h
+++ b/Engine/Source/Editor/Sequencer/Public/SequencerUtilities.h
@@ -350,6 +350,15 @@ struct FSequencerUtilities
 
 	/** Get a movie scene sequence from a FMovideSceneSequenceID */
 	static UE_API UMovieSceneSequence* GetMovieSceneSequence(TSharedPtr<ISequencer>& InSequencer, const FMovieSceneSequenceID& SequenceID);
+
+	/** Snap sections to timeline using source timecode */
+	static UE_API void SnapSectionsToTimelineUsingSourceTimecode(TSharedRef<ISequencer> Sequencer, const TArray<UMovieSceneSection*>& Sections);
+
+	/** Sync section using source timecode */
+	static UE_API void SyncSectionsUsingSourceTimecode(TSharedRef<ISequencer> Sequencer, const TArray<UMovieSceneSection*>& Sections);
+
+	/** Attempts to automatically fix up broken actor references in the current scene */
+	static UE_API void FixActorReferences(TSharedRef<ISequencer> Sequencer);
 };
 
 #undef UE_API
