Sequencer: Moved the Assign Actor submenu extension out of the LevelSequenceEditorSubsystem and into the LevelSequenceCustomization, and the implementation of the various actions into SequencerUtilities (some were already there). This makes it possible for the menu options to provide a specific ISequencer to the implementations, rather than just using the "Active" Sequencer defined by the subsystem.

The public scripting API on the subsystem for these actions still remains, and continues to use the "Active" Sequencer.

#rnx
#Synced-CL 47206406

diff --git a/Engine/Plugins/MovieScene/LevelSequenceEditor/Source/LevelSequenceEditor/Private/LevelSequenceEditorSubsystem.cpp b/Engine/Plugins/MovieScene/LevelSequenceEditor/Source/LevelSequenceEditor/Private/LevelSequenceEditorSubsystem.cpp
--- a/Engine/Plugins/MovieScene/LevelSequenceEditor/Source/LevelSequenceEditor/Private/LevelSequenceEditorSubsystem.cpp
+++ b/Engine/Plugins/MovieScene/LevelSequenceEditor/Source/LevelSequenceEditor/Private/LevelSequenceEditorSubsystem.cpp
@@ -7,20 +7,16 @@
 #include "SequencerCurveEditorObject.h"
 #include "Evaluation/MovieScenePlayback.h"
 #include "ISequencerModule.h"
-#include "Framework/Commands/UICommandList.h"
 #include "LevelSequence.h"
 #include "ISceneOutliner.h"
-#include "LevelSequenceEditorCommands.h"
 #include "MovieScenePossessable.h"
 #include "SequenceBindingTree.h"
-#include "SequencerSettings.h"
 #include "MovieScene.h"
 #include "MovieSceneSpawnable.h"
 #include "SequencerUtilities.h"
 #include "Selection.h"
 #include "Toolkits/AssetEditorToolkit.h"
 
-#include "ActorTreeItem.h"
 #include "PropertyEditorModule.h"
 #include "Widgets/SWindow.h"
 #include "Widgets/Input/SButton.h"
@@ -35,7 +31,6 @@
 #include "Framework/MultiBox/MultiBoxBuilder.h"
 #include "HAL/PlatformApplicationMisc.h"
 #include "Modules/ModuleManager.h"
-#include "SceneOutlinerModule.h"
 #include "ScopedTransaction.h"
 #include "Widgets/Notifications/SNotificationList.h"
 #include "Framework/MultiBox/MultiBoxBuilder.h"
@@ -621,54 +616,9 @@ void ULevelSequenceEditorSubsystem::Initialize(FSubsystemCollectionBase& Collect
 	ISequencerModule& SequencerModule = FModuleManager::Get().LoadModuleChecked<ISequencerModule>("Sequencer");
 	OnSequencerCreatedHandle = SequencerModule.RegisterOnSequencerCreated(FOnSequencerCreated::FDelegate::CreateUObject(this, &ULevelSequenceEditorSubsystem::OnSequencerCreated));
 
-	auto AreActorsSelected = [this]{
-		TArray<AActor*> SelectedActors;
-		GEditor->GetSelectedActors()->GetSelectedObjects<AActor>(SelectedActors); 
-		return SelectedActors.Num() > 0;
-	};
-
-	/* Commands for this subsystem */
-	CommandList = MakeShareable(new FUICommandList);
-
-	CommandList->MapAction(FLevelSequenceEditorCommands::Get().AddActorsToBinding,
-		FExecuteAction::CreateUObject(this, &ULevelSequenceEditorSubsystem::AddActorsToBindingInternal),
-		FCanExecuteAction::CreateLambda(AreActorsSelected));
-
-	CommandList->MapAction(FLevelSequenceEditorCommands::Get().ReplaceBindingWithActors,
-		FExecuteAction::CreateUObject(this, &ULevelSequenceEditorSubsystem::ReplaceBindingWithActorsInternal),
-		FCanExecuteAction::CreateLambda(AreActorsSelected));
-
-	CommandList->MapAction(FLevelSequenceEditorCommands::Get().RemoveActorsFromBinding,
-		FExecuteAction::CreateUObject(this, &ULevelSequenceEditorSubsystem::RemoveActorsFromBindingInternal),
-		FCanExecuteAction::CreateLambda(AreActorsSelected));
-
-	CommandList->MapAction(FLevelSequenceEditorCommands::Get().RemoveAllBindings,
-		FExecuteAction::CreateUObject(this, &ULevelSequenceEditorSubsystem::RemoveAllBindingsInternal));
-
-	CommandList->MapAction(FLevelSequenceEditorCommands::Get().RemoveInvalidBindings,
-		FExecuteAction::CreateUObject(this, &ULevelSequenceEditorSubsystem::RemoveInvalidBindingsInternal));
-
 	/* Menu extenders */
-	AssignActorMenuExtender = MakeShareable(new FExtender);
-	AssignActorMenuExtender->AddMenuExtension("Possessable", EExtensionHook::First, CommandList, FMenuExtensionDelegate::CreateLambda([this](FMenuBuilder& MenuBuilder) {
-
-		if (!IsSelectedBindingRootPossessable())
-		{
-			return;
-		}
-		
-		FFormatNamedArguments Args;
-		MenuBuilder.AddSubMenu(
-			FText::Format(LOCTEXT("AssignActor", "Assign Actor"), Args),
-			FText::Format(LOCTEXT("AssignActorTooltip", "Assign an actor to this track"), Args),
-			FNewMenuDelegate::CreateLambda([this](FMenuBuilder& SubMenuBuilder) { AddAssignActorMenu(SubMenuBuilder); } ));
-		}));
-
-	SequencerModule.GetObjectBindingContextMenuExtensibilityManager()->AddExtender(AssignActorMenuExtender);
-
-	// For now we have the binding properties being a separate menu. When the UX is worked out we will likely merge the AssignActor menu away.
 	BindingPropertiesMenuExtender = MakeShareable(new FExtender);
-	BindingPropertiesMenuExtender->AddMenuExtension("Possessable", EExtensionHook::First, CommandList, FMenuExtensionDelegate::CreateLambda([this](FMenuBuilder& MenuBuilder) {
+	BindingPropertiesMenuExtender->AddMenuExtension("Possessable", EExtensionHook::First, nullptr, FMenuExtensionDelegate::CreateLambda([this](FMenuBuilder& MenuBuilder) {
 
 		FFormatNamedArguments Args;
 		MenuBuilder.AddSubMenu(
@@ -680,7 +630,7 @@ void ULevelSequenceEditorSubsystem::Initialize(FSubsystemCollectionBase& Collect
 	SequencerModule.GetObjectBindingContextMenuExtensibilityManager()->AddExtender(BindingPropertiesMenuExtender);
 
 	RebindComponentMenuExtender = MakeShareable(new FExtender);
-	RebindComponentMenuExtender->AddMenuExtension("Possessable", EExtensionHook::First, CommandList, FMenuExtensionDelegate::CreateLambda([this](FMenuBuilder& MenuBuilder) {
+	RebindComponentMenuExtender->AddMenuExtension("Possessable", EExtensionHook::First, nullptr, FMenuExtensionDelegate::CreateLambda([this](FMenuBuilder& MenuBuilder) {
 		// Only add menu entries where the focused sequence is a ULevelSequence
 		if (!GetActiveSequencer())
 		{
@@ -699,7 +649,7 @@ void ULevelSequenceEditorSubsystem::Initialize(FSubsystemCollectionBase& Collect
 		}
 		}));
 
-	BindingPropertiesMenuExtender->AddMenuExtension("CustomBinding", EExtensionHook::First, CommandList, FMenuExtensionDelegate::CreateLambda([this](FMenuBuilder& MenuBuilder) {
+	BindingPropertiesMenuExtender->AddMenuExtension("CustomBinding", EExtensionHook::First, nullptr, FMenuExtensionDelegate::CreateLambda([this](FMenuBuilder& MenuBuilder) {
 
 		FFormatNamedArguments Args;
 		MenuBuilder.AddSubMenu(
@@ -708,7 +658,7 @@ void ULevelSequenceEditorSubsystem::Initialize(FSubsystemCollectionBase& Collect
 			FNewMenuDelegate::CreateLambda([this](FMenuBuilder& SubMenuBuilder) { AddBindingPropertiesMenu(SubMenuBuilder); }));
 		}));
 
-	BindingPropertiesMenuExtender->AddMenuExtension("ConvertBinding", EExtensionHook::First, CommandList, FMenuExtensionDelegate::CreateLambda([this](FMenuBuilder& MenuBuilder) {
+	BindingPropertiesMenuExtender->AddMenuExtension("ConvertBinding", EExtensionHook::First, nullptr, FMenuExtensionDelegate::CreateLambda([this](FMenuBuilder& MenuBuilder) {
 		 AddConvertBindingsMenu(MenuBuilder);
 		}));
 
@@ -716,7 +666,7 @@ void ULevelSequenceEditorSubsystem::Initialize(FSubsystemCollectionBase& Collect
 
 	SidebarMenuExtender = MakeShared<FExtender>();
 
-	SidebarMenuExtender->AddMenuExtension(TEXT("Possessable"), EExtensionHook::First, CommandList,
+	SidebarMenuExtender->AddMenuExtension(TEXT("Possessable"), EExtensionHook::First, nullptr,
 		FMenuExtensionDelegate::CreateLambda([this](FMenuBuilder& MenuBuilder)
 			{
 				// Only add menu entries where the focused sequence is a ULevelSequence
@@ -728,13 +678,13 @@ void ULevelSequenceEditorSubsystem::Initialize(FSubsystemCollectionBase& Collect
 				AddBindingPropertiesSidebar(MenuBuilder);
 			}));
 
-	SidebarMenuExtender->AddMenuExtension(TEXT("CustomBinding"), EExtensionHook::First, CommandList,
+	SidebarMenuExtender->AddMenuExtension(TEXT("CustomBinding"), EExtensionHook::First, nullptr,
 		FMenuExtensionDelegate::CreateLambda([this](FMenuBuilder& MenuBuilder)
 			{
 				AddBindingPropertiesMenu(MenuBuilder);
 			}));
 
-	SidebarMenuExtender->AddMenuExtension(TEXT("TrackRowMetadata"), EExtensionHook::First, CommandList,
+	SidebarMenuExtender->AddMenuExtension(TEXT("TrackRowMetadata"), EExtensionHook::First, nullptr,
 		FMenuExtensionDelegate::CreateLambda([this](FMenuBuilder& MenuBuilder)
 			{
 				AddTrackRowMetadataMenu(MenuBuilder);
@@ -1445,316 +1395,44 @@ void ULevelSequenceEditorSubsystem::FixActorReferences()
 	}
 }
 
-void ULevelSequenceEditorSubsystem::AddActorsToBindingInternal()
-{
-	TSharedPtr<ISequencer> Sequencer = GetActiveSequencer();
-	if (Sequencer == nullptr)
-	{
-		return;
-	}
-
-	TArray<FGuid> ObjectBindings;
-	Sequencer->GetSelectedObjects(ObjectBindings);
-	if (ObjectBindings.Num() == 0)
-	{
-		return;
-	}
-
-	FMovieSceneBindingProxy BindingProxy(ObjectBindings[0], Sequencer->GetFocusedMovieSceneSequence());
-
-	TArray<AActor*> SelectedActors;
-	GEditor->GetSelectedActors()->GetSelectedObjects<AActor>(SelectedActors); 
-	AddActorsToBinding(SelectedActors, BindingProxy);
-}
-
 void ULevelSequenceEditorSubsystem::AddActorsToBinding(const TArray<AActor*>& Actors, const FMovieSceneBindingProxy& ObjectBinding)
 {
-	TSharedPtr<ISequencer> Sequencer = GetActiveSequencer();
-	if (Sequencer == nullptr)
-	{
-		return;
-	}
-
-	FSequencerUtilities::AddActorsToBinding(Sequencer.ToSharedRef(), Actors, ObjectBinding);
-}
-
-void ULevelSequenceEditorSubsystem::ReplaceBindingWithActorsInternal()
-{
-	TSharedPtr<ISequencer> Sequencer = GetActiveSequencer();
-	if (Sequencer == nullptr)
-	{
-		return;
-	}
-
-	TArray<FGuid> ObjectBindings;
-	Sequencer->GetSelectedObjects(ObjectBindings);
-	if (ObjectBindings.Num() == 0)
+	if (TSharedPtr<ISequencer> Sequencer = GetActiveSequencer())
 	{
-		return;
+		FSequencerUtilities::AddActorsToBinding(Sequencer.ToSharedRef(), Actors, ObjectBinding);
 	}
-
-	FMovieSceneBindingProxy BindingProxy(ObjectBindings[0], Sequencer->GetFocusedMovieSceneSequence());
-
-	TArray<AActor*> SelectedActors;
-	GEditor->GetSelectedActors()->GetSelectedObjects<AActor>(SelectedActors); 
-	ReplaceBindingWithActors(SelectedActors, BindingProxy);
 }
 
 void ULevelSequenceEditorSubsystem::ReplaceBindingWithActors(const TArray<AActor*>& Actors, const FMovieSceneBindingProxy& ObjectBinding)
 {
-	TSharedPtr<ISequencer> Sequencer = GetActiveSequencer();
-	if (Sequencer == nullptr)
-	{
-		return;
-	}
-
-	FSequencerUtilities::ReplaceBindingWithActors(Sequencer.ToSharedRef(), Actors, ObjectBinding);
-}
-
-void ULevelSequenceEditorSubsystem::RemoveActorsFromBindingInternal()
-{
-	TSharedPtr<ISequencer> Sequencer = GetActiveSequencer();
-	if (Sequencer == nullptr)
-	{
-		return;
-	}
-
-	TArray<FGuid> ObjectBindings;
-	Sequencer->GetSelectedObjects(ObjectBindings);
-	if (ObjectBindings.Num() == 0)
+	if (TSharedPtr<ISequencer> Sequencer = GetActiveSequencer())
 	{
-		return;
+		FSequencerUtilities::ReplaceBindingWithActors(Sequencer.ToSharedRef(), Actors, ObjectBinding);
 	}
-
-	FMovieSceneBindingProxy BindingProxy(ObjectBindings[0], Sequencer->GetFocusedMovieSceneSequence());
-
-	TArray<AActor*> SelectedActors;
-	GEditor->GetSelectedActors()->GetSelectedObjects<AActor>(SelectedActors); 
-	RemoveActorsFromBinding(SelectedActors, BindingProxy);
 }
 
 void ULevelSequenceEditorSubsystem::RemoveActorsFromBinding(const TArray<AActor*>& Actors, const FMovieSceneBindingProxy& ObjectBinding)
 {
-	TSharedPtr<ISequencer> Sequencer = GetActiveSequencer();
-	if (Sequencer == nullptr)
-	{
-		return;
-	}
-
-	FSequencerUtilities::RemoveActorsFromBinding(Sequencer.ToSharedRef(), Actors, ObjectBinding);
-}
-
-void ULevelSequenceEditorSubsystem::RemoveAllBindingsInternal()
-{
-	TSharedPtr<ISequencer> Sequencer = GetActiveSequencer();
-	if (Sequencer == nullptr)
-	{
-		return;
-	}
-
-	TArray<FGuid> ObjectBindings;
-	Sequencer->GetSelectedObjects(ObjectBindings);
-	if (ObjectBindings.Num() == 0)
+	if (TSharedPtr<ISequencer> Sequencer = GetActiveSequencer())
 	{
-		return;
+		FSequencerUtilities::RemoveActorsFromBinding(Sequencer.ToSharedRef(), Actors, ObjectBinding);
 	}
-
-	FMovieSceneBindingProxy BindingProxy(ObjectBindings[0], Sequencer->GetFocusedMovieSceneSequence());
-
-	RemoveAllBindings(BindingProxy);
 }
 
 void ULevelSequenceEditorSubsystem::RemoveAllBindings(const FMovieSceneBindingProxy& ObjectBinding)
 {
-	TSharedPtr<ISequencer> Sequencer = GetActiveSequencer();
-	if (Sequencer == nullptr)
-	{
-		return;
-	}
-
-	UMovieSceneSequence* Sequence = Sequencer->GetFocusedMovieSceneSequence();
-	if (!Sequence)
-	{
-		return;
-	}
-
-	UMovieScene* MovieScene = Sequence->GetMovieScene();
-	if (!MovieScene)
-	{
-		return;
-	}
-
-	FScopedTransaction RemoveAllBindings(LOCTEXT("RemoveAllBindings", "Remove All Bound Objects"));
-
-	Sequence->Modify();
-	MovieScene->Modify();
-
-	// Unbind objects
-	FGuid Guid = ObjectBinding.BindingID;
-	Sequence->UnbindPossessableObjects(Guid);
-
-	Sequencer->RestorePreAnimatedState();
-
-	Sequencer->NotifyMovieSceneDataChanged(EMovieSceneDataChangeType::MovieSceneStructureItemsChanged);
-}
-
-void ULevelSequenceEditorSubsystem::RemoveInvalidBindingsInternal()
-{
-	TSharedPtr<ISequencer> Sequencer = GetActiveSequencer();
-	if (Sequencer == nullptr)
-	{
-		return;
-	}
-
-	TArray<FGuid> ObjectBindings;
-	Sequencer->GetSelectedObjects(ObjectBindings);
-	if (ObjectBindings.Num() == 0)
+	if (TSharedPtr<ISequencer> Sequencer = GetActiveSequencer())
 	{
-		return;
+		FSequencerUtilities::RemoveAllBindings(Sequencer.ToSharedRef(), ObjectBinding);
 	}
-
-	FMovieSceneBindingProxy BindingProxy(ObjectBindings[0], Sequencer->GetFocusedMovieSceneSequence());
-
-	RemoveInvalidBindings(BindingProxy);
 }
 
 void ULevelSequenceEditorSubsystem::RemoveInvalidBindings(const FMovieSceneBindingProxy& ObjectBinding)
 {
-	TSharedPtr<ISequencer> Sequencer = GetActiveSequencer();
-	if (Sequencer == nullptr)
-	{
-		return;
-	}
-
-	UMovieSceneSequence* Sequence = Sequencer->GetFocusedMovieSceneSequence();
-	if (!Sequence)
-	{
-		return;
-	}
-
-	UMovieScene* MovieScene = Sequence->GetMovieScene();
-	if (!MovieScene)
-	{
-		return;
-	}
-	
-	FScopedTransaction RemoveInvalidBindings(LOCTEXT("RemoveMissing", "Remove Missing Objects"));
-
-	Sequence->Modify();
-	MovieScene->Modify();
-
-	// Unbind objects
-	FGuid Guid = ObjectBinding.BindingID;
-	Sequence->UnbindInvalidObjects(Guid, Sequencer->GetPlaybackContext());
-
-	// Update label
-	UClass* ActorClass = nullptr;
-
-	TArray<AActor*> ValidActors;
-	for (TWeakObjectPtr<> Ptr : Sequencer->FindObjectsInCurrentSequence(Guid))
-	{
-		if (AActor* Actor = Cast<AActor>(Ptr.Get()))
-		{
-			ActorClass = Actor->GetClass();
-			ValidActors.Add(Actor);
-		}
-	}
-
-	FMovieScenePossessable* Possessable = MovieScene->FindPossessable(Guid);
-	if (Possessable && ActorClass != nullptr && ValidActors.Num() != 0)
-	{
-		if (ValidActors.Num() > 1)
-		{
-			FString NewLabel = ActorClass->GetName() + FString::Printf(TEXT(" (%d)"), ValidActors.Num());
-
-			Possessable->SetName(NewLabel);
-		}
-		else
-		{
-			Possessable->SetName(ValidActors[0]->GetActorLabel());
-		}
-	}
-
-	Sequencer->RestorePreAnimatedState();
-
-	Sequencer->NotifyMovieSceneDataChanged(EMovieSceneDataChangeType::MovieSceneStructureItemsChanged);
-}
-
-void ULevelSequenceEditorSubsystem::AddAssignActorMenu(FMenuBuilder& MenuBuilder)
-{
-	TSharedPtr<ISequencer> Sequencer = GetActiveSequencer();
-	if (Sequencer == nullptr)
-	{
-		return;
-	}
-
-	TArray<FGuid> ObjectBindings;
-	Sequencer->GetSelectedObjects(ObjectBindings);
-	if (ObjectBindings.Num() == 0)
-	{
-		return;
-	}
-
-	MenuBuilder.AddMenuEntry(FLevelSequenceEditorCommands::Get().AddActorsToBinding);
-	MenuBuilder.AddMenuEntry(FLevelSequenceEditorCommands::Get().ReplaceBindingWithActors);
-	MenuBuilder.AddMenuEntry(FLevelSequenceEditorCommands::Get().RemoveActorsFromBinding);
-	MenuBuilder.AddMenuEntry(FLevelSequenceEditorCommands::Get().RemoveAllBindings);
-	MenuBuilder.AddMenuEntry(FLevelSequenceEditorCommands::Get().RemoveInvalidBindings);
-
-	FMovieSceneBindingProxy BindingProxy(ObjectBindings[0], Sequencer->GetFocusedMovieSceneSequence());
-
-	TSet<const AActor*> BoundObjects;
-	{
-		for (TWeakObjectPtr<> Ptr : Sequencer->FindObjectsInCurrentSequence(ObjectBindings[0]))
-		{
-			if (const AActor* Actor = Cast<AActor>(Ptr.Get()))
-			{
-				BoundObjects.Add(Actor);
-			}
-		}
-	}
-
-	auto IsActorValidForAssignment = [BoundObjects](const AActor* InActor){
-		return !BoundObjects.Contains(InActor);
-	};
-
-	// Set up a menu entry to assign an actor to the object binding node
-	FSceneOutlinerInitializationOptions InitOptions;
+	if (TSharedPtr<ISequencer> Sequencer = GetActiveSequencer())
 	{
-		// We hide the header row to keep the UI compact.
-		InitOptions.bShowHeaderRow = false;
-		InitOptions.bShowSearchBox = true;
-		InitOptions.bShowCreateNewFolder = false;
-		InitOptions.bFocusSearchBoxWhenOpened = true;
-		// Only want the actor label column
-		InitOptions.ColumnMap.Add(FSceneOutlinerBuiltInColumnTypes::Label(), FSceneOutlinerColumnInfo(ESceneOutlinerColumnVisibility::Visible, 0));
-		
-		// Only display actors that are not possessed already
-		InitOptions.Filters->AddFilterPredicate<FActorTreeItem>(FActorTreeItem::FFilterPredicate::CreateLambda( IsActorValidForAssignment ) );
+		FSequencerUtilities::RemoveInvalidBindings(Sequencer.ToSharedRef(), ObjectBinding);
 	}
-
-	const float WidthOverride = Sequencer.IsValid() ? Sequencer->GetSequencerSettings()->GetAssetBrowserWidth() : 500.f;
-	const float HeightOverride = Sequencer.IsValid() ? Sequencer->GetSequencerSettings()->GetAssetBrowserHeight() : 400.f;
-
-	// actor selector to allow the user to choose an actor
-	FSceneOutlinerModule& SceneOutlinerModule = FModuleManager::LoadModuleChecked<FSceneOutlinerModule>("SceneOutliner");
-	TSharedRef< SWidget > MiniSceneOutliner =
-		SNew( SBox )
-		.WidthOverride(WidthOverride)
-		.HeightOverride(HeightOverride)
-		[
-			SceneOutlinerModule.CreateActorPicker(
-				InitOptions,
-				FOnActorPicked::CreateLambda([=](AActor* Actor){
-					// Create a new binding for this actor
-					FSlateApplication::Get().DismissAllMenus();
-					FSequencerUtilities::AssignActor(Sequencer.ToSharedRef(), Actor, ObjectBindings[0]);
-				})
-			)
-		];
-
-	MenuBuilder.AddMenuSeparator();
-	MenuBuilder.AddWidget(MiniSceneOutliner, FText::GetEmpty(), true);
 }
 
 void ULevelSequenceEditorSubsystem::AddBindingPropertiesMenu(FMenuBuilder& MenuBuilder)
@@ -2321,32 +1999,6 @@ void ULevelSequenceEditorSubsystem::RebindComponent(const TArray<FMovieSceneBind
 	}
 }
 
-bool ULevelSequenceEditorSubsystem::IsSelectedBindingRootPossessable()
-{
-	TSharedPtr<ISequencer> Sequencer = GetActiveSequencer();
-	if (Sequencer)
-	{
-		TArray<FGuid> ObjectBindings;
-		Sequencer->GetSelectedObjects(ObjectBindings);
-		if (ObjectBindings.Num() > 0)
-		{
-			UMovieSceneSequence* Sequence = Sequencer->GetFocusedMovieSceneSequence();
-			UMovieScene* MovieScene = Sequence ? Sequence->GetMovieScene() : nullptr;
-			if (MovieScene)
-			{
-				if (FMovieScenePossessable* Possessable = MovieScene->FindPossessable(ObjectBindings[0]))
-				{
-					if (!Possessable->GetParent().IsValid() && !Possessable->GetSpawnableObjectBindingID().IsValid())
-					{
-						return true;
-					}
-				}
-			}
-		}
-	}
-	return false;
-}
-
 void ULevelSequenceEditorSubsystem::RefreshBindingDetails(IDetailsView* DetailsView, FGuid ObjectBindingID)
 {
 	if (DetailsView == nullptr)
diff --git a/Engine/Plugins/MovieScene/LevelSequenceEditor/Source/LevelSequenceEditor/Private/Misc/LevelSequenceCustomization.cpp b/Engine/Plugins/MovieScene/LevelSequenceEditor/Source/LevelSequenceEditor/Private/Misc/LevelSequenceCustomization.cpp
--- a/Engine/Plugins/MovieScene/LevelSequenceEditor/Source/LevelSequenceEditor/Private/Misc/LevelSequenceCustomization.cpp
+++ b/Engine/Plugins/MovieScene/LevelSequenceEditor/Source/LevelSequenceEditor/Private/Misc/LevelSequenceCustomization.cpp
@@ -2,6 +2,9 @@
 
 #include "Misc/LevelSequenceCustomization.h"
 
+#include "ActorTreeItem.h"
+#include "Bindings/MovieSceneSpawnableBinding.h"
+#include "ClassViewerFilter.h"
 #include "ClassViewerModule.h"
 #include "Engine/LevelStreaming.h"
 #include "Framework/Application/SlateApplication.h"
@@ -14,16 +17,18 @@
 #include "LevelSequenceFBXInterop.h"
 #include "Modules/ModuleManager.h"
 #include "MovieScene.h"
+#include "MovieSceneBindingReferences.h"
 #include "MVVM/Extensions/IObjectBindingExtension.h"
 #include "MVVM/Selection/Selection.h"
 #include "MVVM/ViewModels/ObjectBindingModel.h"
 #include "MVVM/ViewModels/SequencerEditorViewModel.h"
+#include "SceneOutlinerModule.h"
+#include "SceneOutlinerPublicTypes.h"
 #include "ScopedTransaction.h"
+#include "Selection.h"
 #include "SequencerCommands.h"
+#include "SequencerSettings.h"
 #include "SequencerUtilities.h"
-#include "MovieSceneBindingReferences.h"
-#include "Bindings/MovieSceneSpawnableBinding.h"
-#include "ClassViewerFilter.h"
 
 #define LOCTEXT_NAMESPACE "LevelSequenceCustomization"
 
@@ -39,7 +44,7 @@ void FLevelSequenceCustomization::RegisterSequencerCustomization(FSequencerCusto
 {
 	WeakSequencer = Builder.GetSequencer().AsShared();
 
-	BindActionMenuCommands();
+	BindCommands();
 
 	FSequencerCustomizationInfo Customization;
 	Customization.ActionsMenuExtender = CreateActionsMenuExtender();
@@ -60,41 +65,64 @@ void FLevelSequenceCustomization::UnregisterSequencerCustomization()
 		ExternalCustomization->UnregisterSequencerCustomization();
 	}
 
-	ActionsMenuCommandList = nullptr;
+	CommandList = nullptr;
 	WeakSequencer = nullptr;
 }
 
-void FLevelSequenceCustomization::BindActionMenuCommands()
+void FLevelSequenceCustomization::BindCommands()
 {
 	const FLevelSequenceEditorCommands& Commands = FLevelSequenceEditorCommands::Get();
 
-	ActionsMenuCommandList = MakeShared<FUICommandList>().ToSharedPtr();
+	CommandList = MakeShared<FUICommandList>().ToSharedPtr();
 
-	ActionsMenuCommandList->MapAction(
+	CommandList->MapAction(
 		Commands.ImportFBX,
 		FExecuteAction::CreateRaw(this, &FLevelSequenceCustomization::ImportFBX),
 		FCanExecuteAction::CreateLambda([] { return true; }));
 
-	ActionsMenuCommandList->MapAction(
+	CommandList->MapAction(
 		Commands.ExportFBX,
 		FExecuteAction::CreateRaw(this, &FLevelSequenceCustomization::ExportFBX),
 		FCanExecuteAction::CreateLambda([] { return true; }));
 
 	constexpr int32 MinSectionsRequiredForSnapping = 1;
-	ActionsMenuCommandList->MapAction(
+	CommandList->MapAction(
 		Commands.SnapSectionsToTimelineUsingSourceTimecode,
 		FExecuteAction::CreateRaw(this, &FLevelSequenceCustomization::SnapSectionsToTimelineUsingSourceTimecode),
 		FCanExecuteAction::CreateRaw(this, &FLevelSequenceCustomization::DoesMovieSceneHaveEnoughSectionsSelected, MinSectionsRequiredForSnapping));
 
 	constexpr int32 MinSectionsRequiredForSyncing = 2;
-	ActionsMenuCommandList->MapAction(
+	CommandList->MapAction(
 		Commands.SyncSectionsUsingSourceTimecode,
 		FExecuteAction::CreateRaw(this, &FLevelSequenceCustomization::SyncSectionsUsingSourceTimecode),
 		FCanExecuteAction::CreateRaw(this, &FLevelSequenceCustomization::DoesMovieSceneHaveEnoughSectionsSelected, MinSectionsRequiredForSyncing));
 
-	ActionsMenuCommandList->MapAction(
+	CommandList->MapAction(
 		Commands.FixActorReferences,
 		FExecuteAction::CreateRaw(this, &FLevelSequenceCustomization::FixActorReferences));
+
+	CommandList->MapAction(
+		Commands.AddActorsToBinding,
+		FExecuteAction::CreateRaw(this, &FLevelSequenceCustomization::AddActorsToBinding),
+		FCanExecuteAction::CreateRaw(this, &FLevelSequenceCustomization::AreAnyActorsSelected));
+
+	CommandList->MapAction(
+		Commands.ReplaceBindingWithActors,
+		FExecuteAction::CreateRaw(this, &FLevelSequenceCustomization::ReplaceBindingWithActors),
+		FCanExecuteAction::CreateRaw(this, &FLevelSequenceCustomization::AreAnyActorsSelected));
+
+	CommandList->MapAction(
+		Commands.RemoveActorsFromBinding,
+		FExecuteAction::CreateRaw(this, &FLevelSequenceCustomization::RemoveActorsFromBinding),
+		FCanExecuteAction::CreateRaw(this, &FLevelSequenceCustomization::AreAnyActorsSelected));
+
+	CommandList->MapAction(
+		Commands.RemoveAllBindings,
+		FExecuteAction::CreateRaw(this, &FLevelSequenceCustomization::RemoveAllBindings));
+
+	CommandList->MapAction(
+		Commands.RemoveInvalidBindings,
+		FExecuteAction::CreateRaw(this, &FLevelSequenceCustomization::RemoveInvalidBindings));
 }
 
 TSharedPtr<FExtender> FLevelSequenceCustomization::CreateActionsMenuExtender()
@@ -102,7 +130,7 @@ TSharedPtr<FExtender> FLevelSequenceCustomization::CreateActionsMenuExtender()
 	TSharedPtr<FExtender> ActionsMenuExtender = MakeShared<FExtender>();
 
 	ActionsMenuExtender->AddMenuExtension(
-		"SequenceOptions", EExtensionHook::First, ActionsMenuCommandList,
+		"SequenceOptions", EExtensionHook::First, CommandList,
 		FMenuExtensionDelegate::CreateLambda([](FMenuBuilder& MenuBuilder)
 			{
 				MenuBuilder.AddMenuEntry(FLevelSequenceEditorCommands::Get().ImportFBX);
@@ -110,7 +138,7 @@ TSharedPtr<FExtender> FLevelSequenceCustomization::CreateActionsMenuExtender()
 			}));
 
 	ActionsMenuExtender->AddMenuExtension(
-		"Transform", EExtensionHook::After, ActionsMenuCommandList,
+		"Transform", EExtensionHook::After, CommandList,
 		FMenuExtensionDelegate::CreateLambda([](FMenuBuilder& MenuBuilder)
 			{
 				MenuBuilder.AddMenuEntry(FLevelSequenceEditorCommands::Get().SnapSectionsToTimelineUsingSourceTimecode);
@@ -118,7 +146,7 @@ TSharedPtr<FExtender> FLevelSequenceCustomization::CreateActionsMenuExtender()
 			}));
 
 	ActionsMenuExtender->AddMenuExtension(
-		"Bindings", EExtensionHook::First, ActionsMenuCommandList,
+		"Bindings", EExtensionHook::First, CommandList,
 		FMenuExtensionDelegate::CreateLambda([this](FMenuBuilder& MenuBuilder)
 			{
 				MenuBuilder.AddMenuEntry(FLevelSequenceEditorCommands::Get().FixActorReferences);
@@ -186,7 +214,7 @@ TSharedPtr<FExtender> FLevelSequenceCustomization::CreateObjectBindingContextMen
 	TSharedRef<FExtender> Extender = MakeShared<FExtender>();
 	TSharedPtr<FObjectBindingModel> ObjectBindingModel = InViewModel->CastThisShared<FObjectBindingModel>();
 	Extender->AddMenuExtension(
-			"ObjectBindingActions", EExtensionHook::Before, nullptr,
+			"ObjectBindingActions", EExtensionHook::Before, CommandList,
 			FMenuExtensionDelegate::CreateRaw(this, &FLevelSequenceCustomization::ExtendObjectBindingContextMenu, ObjectBindingModel));
 	return Extender.ToSharedPtr();
 }
@@ -259,8 +287,16 @@ void FLevelSequenceCustomization::ExtendObjectBindingContextMenu(FMenuBuilder& M
 		if (!bCustomBinding)
 		{
 			// Regular possessable
-			// We don't add anything here, but the extension will
 			MenuBuilder.BeginSection("Possessable");
+
+			if (IsSelectedBindingRootPossessable())
+			{
+				MenuBuilder.AddSubMenu(
+					LOCTEXT("AssignActor", "Assign Actor"),
+					LOCTEXT("AssignActorTooltip", "Assign an actor to this track"),
+					FNewMenuDelegate::CreateRaw(this, &FLevelSequenceCustomization::AddAssignActorSubMenu));
+			}
+
 			MenuBuilder.EndSection();
 		}
 		else
@@ -354,6 +390,185 @@ void FLevelSequenceCustomization::ExtendObjectBindingContextMenu(FMenuBuilder& M
 	MenuBuilder.EndSection();
 }
 
+void FLevelSequenceCustomization::AddAssignActorSubMenu(FMenuBuilder& MenuBuilder)
+{
+	MenuBuilder.AddMenuEntry(FLevelSequenceEditorCommands::Get().AddActorsToBinding);
+	MenuBuilder.AddMenuEntry(FLevelSequenceEditorCommands::Get().ReplaceBindingWithActors);
+	MenuBuilder.AddMenuEntry(FLevelSequenceEditorCommands::Get().RemoveActorsFromBinding);
+	MenuBuilder.AddMenuEntry(FLevelSequenceEditorCommands::Get().RemoveAllBindings);
+	MenuBuilder.AddMenuEntry(FLevelSequenceEditorCommands::Get().RemoveInvalidBindings);
+
+	// Set up a menu entry to assign an actor to the object binding node
+	TSharedRef<SWidget> ActorPickerWidget = BuildActorPicker();
+
+	MenuBuilder.AddMenuSeparator();
+	MenuBuilder.AddWidget(ActorPickerWidget, FText::GetEmpty(), true);
+}
+
+TSharedRef<SWidget> FLevelSequenceCustomization::BuildActorPicker()
+{
+	TSharedPtr<ISequencer> Sequencer = WeakSequencer.Pin();
+	if (!Sequencer)
+	{
+		return SNullWidget::NullWidget;
+	}
+
+	TArray<FGuid> ObjectBindings;
+	Sequencer->GetSelectedObjects(ObjectBindings);
+	if (ObjectBindings.IsEmpty())
+	{
+		return SNullWidget::NullWidget;
+	}
+	const FGuid SelectedObjectBinding = ObjectBindings[0];
+
+	// Find the set of already bound actors and create a filter to exclude those from the actor picker widget
+	TSet<const AActor*> AlreadyBoundActors;
+	for (TWeakObjectPtr<> Ptr : Sequencer->FindObjectsInCurrentSequence(SelectedObjectBinding))
+	{
+		if (const AActor* Actor = Cast<AActor>(Ptr.Get()))
+		{
+			AlreadyBoundActors.Add(Actor);
+		}
+	}
+
+	auto IsActorValidForAssignment = [AlreadyBoundActors](const AActor* InActor)
+	{
+		return !AlreadyBoundActors.Contains(InActor);
+	};
+
+	FSceneOutlinerInitializationOptions InitOptions;
+	InitOptions.bShowHeaderRow = false;
+	InitOptions.bShowSearchBox = true;
+	InitOptions.bShowCreateNewFolder = false;
+	InitOptions.bFocusSearchBoxWhenOpened = true;
+
+	// Only display the actor label column
+	InitOptions.ColumnMap.Add(FSceneOutlinerBuiltInColumnTypes::Label(), FSceneOutlinerColumnInfo(ESceneOutlinerColumnVisibility::Visible, 0));
+
+	// Only display actors that are not possessed already
+	InitOptions.Filters->AddFilterPredicate<FActorTreeItem>(FActorTreeItem::FFilterPredicate::CreateLambda(IsActorValidForAssignment));
+
+	const float WidthOverride = Sequencer->GetSequencerSettings() ? Sequencer->GetSequencerSettings()->GetAssetBrowserWidth() : 500.0f;
+	const float HeightOverride = Sequencer->GetSequencerSettings() ? Sequencer->GetSequencerSettings()->GetAssetBrowserHeight() : 400.0f;
+
+	FSceneOutlinerModule& SceneOutlinerModule = FModuleManager::LoadModuleChecked<FSceneOutlinerModule>("SceneOutliner");
+
+	return SNew(SBox)
+		.WidthOverride(WidthOverride)
+		.HeightOverride(HeightOverride)
+		[
+			SceneOutlinerModule.CreateActorPicker(
+				InitOptions,
+				FOnActorPicked::CreateLambda([Sequencer, SelectedObjectBinding](AActor* Actor)
+				{
+					// Create a new binding for this actor
+					FSlateApplication::Get().DismissAllMenus();
+					FSequencerUtilities::AssignActor(Sequencer.ToSharedRef(), Actor, SelectedObjectBinding);
+				})
+			)
+		];
+}
+
+void FLevelSequenceCustomization::AddActorsToBinding()
+{
+	TArray<AActor*> SelectedActors = GetSelectedActors();
+	const FMovieSceneBindingProxy BindingProxy = GetSelectedBindingProxy();
+
+	if (WeakSequencer.IsValid() && BindingProxy.BindingID.IsValid())
+	{
+		FSequencerUtilities::AddActorsToBinding(WeakSequencer.Pin().ToSharedRef(), SelectedActors, BindingProxy);
+	}
+}
+
+void FLevelSequenceCustomization::ReplaceBindingWithActors()
+{
+	TArray<AActor*> SelectedActors = GetSelectedActors();
+	const FMovieSceneBindingProxy BindingProxy = GetSelectedBindingProxy();
+
+	if (WeakSequencer.IsValid() && BindingProxy.BindingID.IsValid())
+	{
+		FSequencerUtilities::ReplaceBindingWithActors(WeakSequencer.Pin().ToSharedRef(), SelectedActors, BindingProxy);
+	}
+}
+
+void FLevelSequenceCustomization::RemoveActorsFromBinding()
+{
+	TArray<AActor*> SelectedActors = GetSelectedActors();
+	const FMovieSceneBindingProxy BindingProxy = GetSelectedBindingProxy();
+
+	if (WeakSequencer.IsValid() && BindingProxy.BindingID.IsValid())
+	{
+		FSequencerUtilities::RemoveActorsFromBinding(WeakSequencer.Pin().ToSharedRef(), SelectedActors, BindingProxy);
+	}
+}
+
+void FLevelSequenceCustomization::RemoveAllBindings()
+{
+	const FMovieSceneBindingProxy BindingProxy = GetSelectedBindingProxy();
+	if (WeakSequencer.IsValid() && BindingProxy.BindingID.IsValid())
+	{
+		FSequencerUtilities::RemoveAllBindings(WeakSequencer.Pin().ToSharedRef(), BindingProxy);
+	}
+}
+
+void FLevelSequenceCustomization::RemoveInvalidBindings()
+{
+	const FMovieSceneBindingProxy BindingProxy = GetSelectedBindingProxy();
+	if (WeakSequencer.IsValid() && BindingProxy.BindingID.IsValid())
+	{
+		FSequencerUtilities::RemoveInvalidBindings(WeakSequencer.Pin().ToSharedRef(), BindingProxy);
+	}
+}
+
+FMovieSceneBindingProxy FLevelSequenceCustomization::GetSelectedBindingProxy()
+{
+	if (TSharedPtr<ISequencer> Sequencer = WeakSequencer.Pin())
+	{
+		TArray<FGuid> ObjectBindings;
+		Sequencer->GetSelectedObjects(ObjectBindings);
+		if (!ObjectBindings.IsEmpty())
+		{
+			return FMovieSceneBindingProxy(ObjectBindings[0], Sequencer->GetFocusedMovieSceneSequence());
+		}
+	}
+
+	return FMovieSceneBindingProxy();
+}
+
+TArray<AActor*> FLevelSequenceCustomization::GetSelectedActors()
+{
+	TArray<AActor*> SelectedActors;
+	if (USelection* ActorSelection = GEditor->GetSelectedActors())
+	{
+		ActorSelection->GetSelectedObjects<AActor>(SelectedActors);
+	}
+	return SelectedActors;
+}
+
+bool FLevelSequenceCustomization::AreAnyActorsSelected()
+{
+	TArray<AActor*> SelectedActors = GetSelectedActors();
+	return !SelectedActors.IsEmpty();
+}
+
+bool FLevelSequenceCustomization::IsSelectedBindingRootPossessable()
+{
+	const FMovieSceneBindingProxy BindingProxy = GetSelectedBindingProxy();
+	UMovieScene* MovieScene = BindingProxy.GetMovieScene();
+
+	if (BindingProxy.BindingID.IsValid() && MovieScene)
+	{
+		if (FMovieScenePossessable* Possessable = MovieScene->FindPossessable(BindingProxy.BindingID))
+		{
+			if (!Possessable->GetParent().IsValid() && !Possessable->GetSpawnableObjectBindingID().IsValid())
+			{
+				return true;
+			}
+		}
+	}
+	return false;
+}
+
 TSharedPtr<FExtender> FLevelSequenceCustomization::CreateObjectBindingSidebarMenuExtender(FViewModelPtr InViewModel)
 {
 	TSharedRef<FExtender> Extender = MakeShared<FExtender>();
diff --git a/Engine/Plugins/MovieScene/LevelSequenceEditor/Source/LevelSequenceEditor/Private/Misc/LevelSequenceCustomization.h b/Engine/Plugins/MovieScene/LevelSequenceEditor/Source/LevelSequenceEditor/Private/Misc/LevelSequenceCustomization.h
--- a/Engine/Plugins/MovieScene/LevelSequenceEditor/Source/LevelSequenceEditor/Private/Misc/LevelSequenceCustomization.h
+++ b/Engine/Plugins/MovieScene/LevelSequenceEditor/Source/LevelSequenceEditor/Private/Misc/LevelSequenceCustomization.h
@@ -28,29 +28,48 @@ protected:
 
 private:
 
-	// Actions menu extensions
-	void BindActionMenuCommands();
+	/** Map the UI commands used by the extensions in this customization */
+	void BindCommands();
+
+	/** Create the menu extenders that will be added to the registered customization builder */
 	TSharedPtr<FExtender> CreateActionsMenuExtender();
+	TSharedPtr<FExtender> CreateObjectBindingContextMenuExtender(FViewModelPtr InViewModel);
+	TSharedPtr<FExtender> CreateObjectBindingSidebarMenuExtender(FViewModelPtr InViewModel);
+
+	/** Action Menu Extensions */
 	void ImportFBX();
 	void ExportFBX();
 	void SnapSectionsToTimelineUsingSourceTimecode();
 	void SyncSectionsUsingSourceTimecode();
-	bool DoesMovieSceneHaveEnoughSectionsSelected(int32 MinSectionCount);
 	void FixActorReferences();
 
-	// Object binding context menu extensions
-	TSharedPtr<FExtender> CreateObjectBindingContextMenuExtender(FViewModelPtr InViewModel);
+	bool DoesMovieSceneHaveEnoughSectionsSelected(int32 MinSectionCount);
+
+	/** Object Binding Context Menu Extensions */
 	void ExtendObjectBindingContextMenu(FMenuBuilder& MenuBuilder, TSharedPtr<FObjectBindingModel> ObjectBindingModel);
+	void AddAssignActorSubMenu(FMenuBuilder& MenuBuilder);
+	TSharedRef<SWidget> BuildActorPicker();
 
-	// Object binding sidebar menu extensions
-	TSharedPtr<FExtender> CreateObjectBindingSidebarMenuExtender(FViewModelPtr InViewModel);
+	void AddActorsToBinding();
+	void ReplaceBindingWithActors();
+	void RemoveActorsFromBinding();
+	void RemoveAllBindings();
+	void RemoveInvalidBindings();
+
+	FMovieSceneBindingProxy GetSelectedBindingProxy();
+	TArray<AActor*> GetSelectedActors();
+
+	bool AreAnyActorsSelected();
+	bool IsSelectedBindingRootPossessable();
+
+	/** Object Binding Sidebar Menu Extensions */
 	void ExtendObjectBindingSidebarMenu(FMenuBuilder& MenuBuilder, TSharedPtr<FObjectBindingModel> ObjectBindingModel);
 
 private:
 
 	TWeakPtr<ISequencer> WeakSequencer;
 
-	TSharedPtr<FUICommandList> ActionsMenuCommandList;
+	TSharedPtr<FUICommandList> CommandList;
 
 	TArray<TUniquePtr<ISequencerCustomization>>	AdditionalCustomizations;
 };
diff --git a/Engine/Plugins/MovieScene/LevelSequenceEditor/Source/LevelSequenceEditor/Public/LevelSequenceEditorSubsystem.h b/Engine/Plugins/MovieScene/LevelSequenceEditor/Source/LevelSequenceEditor/Public/LevelSequenceEditorSubsystem.h
--- a/Engine/Plugins/MovieScene/LevelSequenceEditor/Source/LevelSequenceEditor/Public/LevelSequenceEditorSubsystem.h
+++ b/Engine/Plugins/MovieScene/LevelSequenceEditor/Source/LevelSequenceEditor/Public/LevelSequenceEditorSubsystem.h
@@ -309,14 +309,8 @@ private:
 
 	UE_API TSharedPtr<ISequencer> GetActiveSequencer();
 	
-	UE_API void AddActorsToBindingInternal();
-	UE_API void ReplaceBindingWithActorsInternal();
-	UE_API void RemoveActorsFromBindingInternal();
-	UE_API void RemoveAllBindingsInternal();
-	UE_API void RemoveInvalidBindingsInternal();
 	UE_API void RebindComponentInternal(const FName& ComponentName);
 
-	UE_API void AddAssignActorMenu(FMenuBuilder& MenuBuilder);
 	UE_API void AddBindingPropertiesMenu(FMenuBuilder& MenuBuilder);
 	UE_API void AddConvertBindingsMenu(FMenuBuilder& MenuBuilder);
 
@@ -345,8 +339,6 @@ private:
 	UE_API void GetRebindComponentNames(TArray<FName>& OutComponentNames);
 	UE_API void RebindComponentMenu(FMenuBuilder& MenuBuilder);
 
-	UE_API bool IsSelectedBindingRootPossessable();
-
 	FDelegateHandle OnSequencerCreatedHandle;
 
 	/* List of sequencers that have been created */
@@ -358,12 +350,6 @@ private:
 	UPROPERTY()
 	TArray<TObjectPtr<USequencerCurveEditorObject>> CurveEditorArray;
 
-	TSharedPtr<FUICommandList> CommandList;
-
-	TSharedPtr<FExtender> TransformMenuExtender;
-	TSharedPtr<FExtender> FixActorReferencesMenuExtender;
-
-	TSharedPtr<FExtender> AssignActorMenuExtender;
 	TSharedPtr<FExtender> BindingPropertiesMenuExtender;
 	TSharedPtr<FExtender> RebindComponentMenuExtender;
 	TSharedPtr<FExtender> SidebarMenuExtender;
diff --git a/Engine/Source/Editor/Sequencer/Private/SequencerUtilities.cpp b/Engine/Source/Editor/Sequencer/Private/SequencerUtilities.cpp
--- a/Engine/Source/Editor/Sequencer/Private/SequencerUtilities.cpp
+++ b/Engine/Source/Editor/Sequencer/Private/SequencerUtilities.cpp
@@ -4342,6 +4342,78 @@ void FSequencerUtilities::RemoveActorsFromBinding(TSharedRef<ISequencer> Sequenc
 	Sequencer->NotifyMovieSceneDataChanged(EMovieSceneDataChangeType::MovieSceneStructureItemsChanged);
 }
 
+void FSequencerUtilities::RemoveAllBindings(TSharedRef<ISequencer> Sequencer, const FMovieSceneBindingProxy& ObjectBinding)
+{
+	UMovieSceneSequence* Sequence = ObjectBinding.Sequence;
+	UMovieScene* MovieScene = ObjectBinding.GetMovieScene();
+	if (!Sequence || !MovieScene)
+	{
+		return;
+	}
+
+	FScopedTransaction RemoveAllBindings(LOCTEXT("RemoveAllBindings", "Remove All Bound Objects"));
+
+	Sequence->Modify();
+	MovieScene->Modify();
+
+	FGuid Guid = ObjectBinding.BindingID;
+	Sequence->UnbindPossessableObjects(Guid);
+
+	Sequencer->RestorePreAnimatedState();
+
+	Sequencer->NotifyMovieSceneDataChanged(EMovieSceneDataChangeType::MovieSceneStructureItemsChanged);
+}
+
+void FSequencerUtilities::RemoveInvalidBindings(TSharedRef<ISequencer> Sequencer, const FMovieSceneBindingProxy& ObjectBinding)
+{
+	UMovieSceneSequence* Sequence = ObjectBinding.Sequence;
+	UMovieScene* MovieScene = ObjectBinding.GetMovieScene();
+	if (!Sequence || !MovieScene)
+	{
+		return;
+	}
+
+	FScopedTransaction RemoveInvalidBindings(LOCTEXT("RemoveMissing", "Remove Missing Objects"));
+
+	Sequence->Modify();
+	MovieScene->Modify();
+
+	FGuid Guid = ObjectBinding.BindingID;
+	Sequence->UnbindInvalidObjects(Guid, Sequencer->GetPlaybackContext());
+
+	// Update label
+	UClass* ActorClass = nullptr;
+
+	TArray<AActor*> ValidActors;
+	for (TWeakObjectPtr<> Ptr : Sequencer->FindObjectsInCurrentSequence(Guid))
+	{
+		if (AActor* Actor = Cast<AActor>(Ptr.Get()))
+		{
+			ActorClass = Actor->GetClass();
+			ValidActors.Add(Actor);
+		}
+	}
+
+	FMovieScenePossessable* Possessable = MovieScene->FindPossessable(Guid);
+	if (Possessable && ActorClass != nullptr && ValidActors.Num() != 0)
+	{
+		if (ValidActors.Num() > 1)
+		{
+			FString NewLabel = ActorClass->GetName() + FString::Printf(TEXT(" (%d)"), ValidActors.Num());
+
+			Possessable->SetName(NewLabel);
+		}
+		else
+		{
+			Possessable->SetName(ValidActors[0]->GetActorLabel());
+		}
+	}
+
+	Sequencer->RestorePreAnimatedState();
+
+	Sequencer->NotifyMovieSceneDataChanged(EMovieSceneDataChangeType::MovieSceneStructureItemsChanged);
+}
+
 struct FBakeData
 {
 	TArray<FVector> Locations;
diff --git a/Engine/Source/Editor/Sequencer/Public/SequencerUtilities.h b/Engine/Source/Editor/Sequencer/Public/SequencerUtilities.h
--- a/Engine/Source/Editor/Sequencer/Public/SequencerUtilities.h
+++ b/Engine/Source/Editor/Sequencer/Public/SequencerUtilities.h
@@ -329,6 +329,8 @@ struct FSequencerUtilities
 	static UE_API void AddObjectsToBinding(TSharedRef<ISequencer> Sequencer, const TArray<UObject*>& Objects, const FMovieSceneBindingProxy& ObjectBinding, UObject* ResolutionContext);
 	static UE_API void ReplaceBindingWithActors(TSharedRef<ISequencer> Sequencer, const TArray<AActor*>& Actors, const FMovieSceneBindingProxy& ObjectBinding);
 	static UE_API void RemoveActorsFromBinding(TSharedRef<ISequencer> Sequencer, const TArray<AActor*>& Actors, const FMovieSceneBindingProxy& ObjectBinding);
+	static UE_API void RemoveAllBindings(TSharedRef<ISequencer> Sequencer, const FMovieSceneBindingProxy& ObjectBinding);
+	static UE_API void RemoveInvalidBindings(TSharedRef<ISequencer> Sequencer, const FMovieSceneBindingProxy& ObjectBinding);
 
 	static UE_API bool BakeTransform(TSharedRef<ISequencer> Sequencer, const TArray<FMovieSceneBindingProxy>& ObjectBindings, const FBakingAnimationKeySettings& InSettings);
 
