Sequencer: Moved the Object Bindings submenu and sidebar extension out of the LevelSequenceEditorSubsystem and into the LevelSequenceCustomization. This makes it possible to provide a specific ISequencer to the implementations, rather than just using the "Active" Sequencer defined by the subsystem.

#rnx
#Synced-CL 47206406

diff --git a/Engine/Plugins/MovieScene/LevelSequenceEditor/Source/LevelSequenceEditor/Private/LevelSequenceEditorSubsystem.cpp b/Engine/Plugins/MovieScene/LevelSequenceEditor/Source/LevelSequenceEditor/Private/LevelSequenceEditorSubsystem.cpp
--- a/Engine/Plugins/MovieScene/LevelSequenceEditor/Source/LevelSequenceEditor/Private/LevelSequenceEditorSubsystem.cpp
+++ b/Engine/Plugins/MovieScene/LevelSequenceEditor/Source/LevelSequenceEditor/Private/LevelSequenceEditorSubsystem.cpp
@@ -343,12 +343,15 @@ TSharedRef<SWidget> FMovieSceneBindingPropertyInfoListCustomization::GenerateBin
 		{
 			Bindings.Add({ BindingGuid, BindingIndex });
 		}
-		LevelSequenceEditorSubsystem->AddChangeBindingTypeMenu(MenuBuilder, Sequencer, Bindings, false, [this, Sequencer, &DetailBuilder]()
+		LevelSequenceEditorSubsystem->AddChangeBindingTypeMenu(MenuBuilder, Sequencer, Bindings, false, [this, &DetailBuilder]()
 			{
 				if (TSharedPtr<IDetailsView> DetailsView  = DetailBuilder.GetDetailsViewSharedPtr())
 				{
-					LevelSequenceEditorSubsystem->RefreshBindingDetails(DetailsView.Get(), BindingGuid);
-					LevelSequenceEditorSubsystem->OnFinishedChangingLocators(FPropertyChangedEvent(nullptr), DetailsView.ToSharedRef(), BindingGuid);
+					if (TSharedPtr<ISequencer> Sequencer = SequencerPtr.Pin())
+					{
+						LevelSequenceEditorSubsystem->RefreshBindingDetails(Sequencer.ToSharedRef(), DetailsView.Get(), BindingGuid);
+						LevelSequenceEditorSubsystem->OnFinishedChangingLocators(FPropertyChangedEvent(nullptr), Sequencer.ToSharedRef(), DetailsView.ToSharedRef(), BindingGuid);
+					}
 				}
 			});
 	}
@@ -426,7 +429,10 @@ TSharedRef<SWidget> FMovieSceneBindingPropertyInfoListCustomization::OnGetConver
 			{
 				if (TSharedPtr<IDetailsView> DetailsView = DetailBuilder.GetDetailsViewSharedPtr())
 				{
-					LevelSequenceEditorSubsystem->RefreshBindingDetails(DetailsView.Get(), BindingGuid);
+					if (TSharedPtr<ISequencer> Sequencer = SequencerPtr.Pin())
+					{
+						LevelSequenceEditorSubsystem->RefreshBindingDetails(Sequencer.ToSharedRef(), DetailsView.Get(), BindingGuid);
+					}
 				}
 			});
 	}
@@ -602,7 +608,10 @@ TSharedRef<SWidget> FMovieSceneBindingPropertyInfoDetailCustomization::OnGetChan
 			{
 				if (TSharedPtr<IDetailsView> DetailsView = StructBuilder.GetParentCategory().GetParentLayout().GetDetailsViewSharedPtr())
 				{
-					LevelSequenceEditorSubsystem->RefreshBindingDetails(DetailsView.Get(), BindingGuid);
+					if (TSharedPtr<ISequencer> Sequencer = SequencerPtr.Pin())
+					{
+						LevelSequenceEditorSubsystem->RefreshBindingDetails(Sequencer.ToSharedRef(), DetailsView.Get(), BindingGuid);
+					}
 				}
 			});
 	}
@@ -617,18 +626,6 @@ void ULevelSequenceEditorSubsystem::Initialize(FSubsystemCollectionBase& Collect
 	OnSequencerCreatedHandle = SequencerModule.RegisterOnSequencerCreated(FOnSequencerCreated::FDelegate::CreateUObject(this, &ULevelSequenceEditorSubsystem::OnSequencerCreated));
 
 	/* Menu extenders */
-	BindingPropertiesMenuExtender = MakeShareable(new FExtender);
-	BindingPropertiesMenuExtender->AddMenuExtension("Possessable", EExtensionHook::First, nullptr, FMenuExtensionDelegate::CreateLambda([this](FMenuBuilder& MenuBuilder) {
-
-		FFormatNamedArguments Args;
-		MenuBuilder.AddSubMenu(
-			FText::Format(LOCTEXT("BindingProperties", "Binding Properties"), Args),
-			FText::Format(LOCTEXT("BindingPropertiesTooltip", "Modify the actor and object bindings for this track"), Args),
-			FNewMenuDelegate::CreateLambda([this](FMenuBuilder& SubMenuBuilder) { AddBindingPropertiesMenu(SubMenuBuilder); }));
-		}));
-
-	SequencerModule.GetObjectBindingContextMenuExtensibilityManager()->AddExtender(BindingPropertiesMenuExtender);
-
 	RebindComponentMenuExtender = MakeShareable(new FExtender);
 	RebindComponentMenuExtender->AddMenuExtension("Possessable", EExtensionHook::First, nullptr, FMenuExtensionDelegate::CreateLambda([this](FMenuBuilder& MenuBuilder) {
 		// Only add menu entries where the focused sequence is a ULevelSequence
@@ -649,41 +646,10 @@ void ULevelSequenceEditorSubsystem::Initialize(FSubsystemCollectionBase& Collect
 		}
 		}));
 
-	BindingPropertiesMenuExtender->AddMenuExtension("CustomBinding", EExtensionHook::First, nullptr, FMenuExtensionDelegate::CreateLambda([this](FMenuBuilder& MenuBuilder) {
-
-		FFormatNamedArguments Args;
-		MenuBuilder.AddSubMenu(
-			FText::Format(LOCTEXT("BindingProperties", "Binding Properties"), Args),
-			FText::Format(LOCTEXT("BindingPropertiesTooltip", "Modify the actor and object bindings for this track"), Args),
-			FNewMenuDelegate::CreateLambda([this](FMenuBuilder& SubMenuBuilder) { AddBindingPropertiesMenu(SubMenuBuilder); }));
-		}));
-
-	BindingPropertiesMenuExtender->AddMenuExtension("ConvertBinding", EExtensionHook::First, nullptr, FMenuExtensionDelegate::CreateLambda([this](FMenuBuilder& MenuBuilder) {
-		 AddConvertBindingsMenu(MenuBuilder);
-		}));
-
 	SequencerModule.GetObjectBindingContextMenuExtensibilityManager()->AddExtender(RebindComponentMenuExtender);
 
 	SidebarMenuExtender = MakeShared<FExtender>();
 
-	SidebarMenuExtender->AddMenuExtension(TEXT("Possessable"), EExtensionHook::First, nullptr,
-		FMenuExtensionDelegate::CreateLambda([this](FMenuBuilder& MenuBuilder)
-			{
-				// Only add menu entries where the focused sequence is a ULevelSequence
-				if (!GetActiveSequencer())
-				{
-					return;
-				}
-
-				AddBindingPropertiesSidebar(MenuBuilder);
-			}));
-
-	SidebarMenuExtender->AddMenuExtension(TEXT("CustomBinding"), EExtensionHook::First, nullptr,
-		FMenuExtensionDelegate::CreateLambda([this](FMenuBuilder& MenuBuilder)
-			{
-				AddBindingPropertiesMenu(MenuBuilder);
-			}));
-
 	SidebarMenuExtender->AddMenuExtension(TEXT("TrackRowMetadata"), EExtensionHook::First, nullptr,
 		FMenuExtensionDelegate::CreateLambda([this](FMenuBuilder& MenuBuilder)
 			{
@@ -703,7 +669,7 @@ void ULevelSequenceEditorSubsystem::Deinitialize()
 		SequencerModulePtr->UnregisterOnSequencerCreated(OnSequencerCreatedHandle);
 	}
 
-	BindingPropertyInfoList = nullptr;
+	BindingPropertyInfoLists.Empty();
 	TrackRowMetadataHelperList.Empty();
 	if (FSlateApplication::IsInitialized())
 	{
@@ -718,11 +684,25 @@ void ULevelSequenceEditorSubsystem::OnSequencerCreated(TSharedRef<ISequencer> In
 
 	Sequencers.Add(TWeakPtr<ISequencer>(InSequencer));
 	InSequencer->OnCloseEvent().AddUObject(this, &ULevelSequenceEditorSubsystem::OnSequencerClosed);
+
+	// Create a new binding list that will be used with the sequencer that was just added
+	UMovieSceneBindingPropertyInfoList* BindingPropertyInfoList = GetBindingPropertyInfoList(InSequencer);
+	if (!BindingPropertyInfoList)
+	{
+		UMovieSceneBindingPropertyInfoList* NewBindingPropertyInfoList = NewObject<UMovieSceneBindingPropertyInfoList>(this);
+		NewBindingPropertyInfoList->Sequencer = InSequencer;
+		BindingPropertyInfoLists.Add(NewBindingPropertyInfoList);
+	}
 }
 
 void ULevelSequenceEditorSubsystem::OnSequencerClosed(TSharedRef<ISequencer> InSequencer)
 {
-	BindingPropertyInfoList = nullptr;
+	const uint32 IndexToRemove = BindingPropertyInfoLists.IndexOfByPredicate([InSequencer](const TObjectPtr<UMovieSceneBindingPropertyInfoList>& List) { return List->Sequencer == InSequencer; });
+	if (IndexToRemove != INDEX_NONE)
+	{
+		BindingPropertyInfoLists.RemoveAt(IndexToRemove);
+	}
+
 	TrackRowMetadataHelperList.Empty();
 }
 
@@ -768,23 +748,6 @@ void ULevelSequenceEditorSubsystem::AddTrackRowMetadataCustomizations(TSharedRef
 	}
 }
 
-void ULevelSequenceEditorSubsystem::OnBindingPropertyMenuBeingDestroyed(const TSharedRef<IMenu>& Menu, TSharedRef<IDetailsView> DetailsView)
-{
-	TSharedPtr<SWidget> ContentWidget = Menu->GetContent();
-	TSharedPtr<SWidget> ParentWidget = DetailsView;
-	while (ParentWidget)
-	{
-		if (ParentWidget == ContentWidget)
-		{
-			// Binding Properties Menu has closed, clear the binding property list
-			BindingPropertyInfoList = nullptr;
-			FSlateApplication::Get().OnMenuBeingDestroyed().RemoveAll(this);
-			break;
-		}
-		ParentWidget = ParentWidget->GetParentWidget();
-	}
-}
-
 void ULevelSequenceEditorSubsystem::OnTrackRowMetadataMenuBeingDestroyed(const TSharedRef<IMenu>& Menu, TSharedRef<IDetailsView> DetailsView)
 {
 	TSharedPtr<SWidget> ContentWidget = Menu->GetContent();
@@ -1437,105 +1400,14 @@ void ULevelSequenceEditorSubsystem::RemoveInvalidBindings(const FMovieSceneBindi
 
 void ULevelSequenceEditorSubsystem::AddBindingPropertiesMenu(FMenuBuilder& MenuBuilder)
 {
-	TSharedPtr<ISequencer> Sequencer = GetActiveSequencer();
-	if (Sequencer == nullptr)
-	{
-		return;
-	}
-
-	TArray<FGuid> ObjectBindings;
-	Sequencer->GetSelectedObjects(ObjectBindings);
-	if (ObjectBindings.Num() == 0)
-	{
-		return;
-	}
-
-	UMovieSceneSequence* Sequence = Sequencer->GetFocusedMovieSceneSequence();
-	if (!Sequence)
-	{
-		return;
-	}
-	if (FMovieSceneBindingReferences* BindingReferences = Sequence->GetBindingReferences())
+	if (TSharedPtr<ISequencer> Sequencer = GetActiveSequencer())
 	{
+		TSharedRef<SWidget> DetailsView = BuildBindingPropertiesDetailsWidget(Sequencer.ToSharedRef());
 		MenuBuilder.AddMenuSeparator();
-
-		NotifyHook = FBindingPropertiesNotifyHook(Sequence);
-		// Set up a details panel for the list of locators
-		FDetailsViewArgs DetailsViewArgs;
-		{
-			DetailsViewArgs.bAllowSearch = false;
-			DetailsViewArgs.bCustomFilterAreaLocation = true;
-			DetailsViewArgs.bCustomNameAreaLocation = true;
-			DetailsViewArgs.bHideSelectionTip = true;
-			DetailsViewArgs.bLockable = false;
-			DetailsViewArgs.bSearchInitialKeyFocus = true;
-			DetailsViewArgs.bUpdatesFromSelection = false;
-			DetailsViewArgs.bShowOptions = false;
-			DetailsViewArgs.bShowModifiedPropertiesOption = false;
-			DetailsViewArgs.bShowScrollBar = false;
-			DetailsViewArgs.bAllowMultipleTopLevelObjects = false;
-			DetailsViewArgs.NotifyHook = &NotifyHook;
-			DetailsViewArgs.DefaultsOnlyVisibility = EEditDefaultsOnlyNodeVisibility::Automatic;
-		}
-
-		TSharedRef<IDetailsView> DetailsView = FModuleManager::GetModuleChecked<FPropertyEditorModule>("PropertyEditor").CreateDetailView(DetailsViewArgs);
-
-		AddBindingDetailCustomizations(DetailsView, Sequencer, ObjectBindings[0]);
-
-		RefreshBindingDetails(&DetailsView.Get(), ObjectBindings[0]);
-		DetailsView->OnFinishedChangingProperties().AddUObject(this, &ULevelSequenceEditorSubsystem::OnFinishedChangingLocators, DetailsView, ObjectBindings[0]);
-
-		if (FSlateApplication::IsInitialized())
-		{
-			FSlateApplication::Get().OnMenuBeingDestroyed().AddUObject(this, &ULevelSequenceEditorSubsystem::OnBindingPropertyMenuBeingDestroyed, DetailsView);
-		}
 		MenuBuilder.AddWidget(DetailsView, FText::GetEmpty(), true);
 	}
 }
 
-void ULevelSequenceEditorSubsystem::AddConvertBindingsMenu(FMenuBuilder& MenuBuilder)
-{
-	// Binding conversion
-	
-	MenuBuilder.AddSubMenu(
-		LOCTEXT("ConvertBindingLabel", "Convert Selected Binding(s) To..."),
-		LOCTEXT("ConvertBindingLabelTooltip", "Convert selected bindings into another binding type"),
-		FNewMenuDelegate::CreateLambda([this](FMenuBuilder& MenuBuilder)
-			{
-				TSharedPtr<ISequencer> Sequencer = GetActiveSequencer();
-				if (Sequencer == nullptr)
-				{
-					return;
-				}
-
-				UMovieSceneSequence* const Sequence = Sequencer->GetFocusedMovieSceneSequence();
-				if (!IsValid(Sequence))
-				{
-					return;
-				}
-
-				TArray<FGuid> ObjectBindings;
-				Sequencer->GetSelectedObjects(ObjectBindings);
-				if (ObjectBindings.Num() == 0)
-				{
-					return;
-				}
-
-				TArray<FSequencerChangeBindingInfo> Bindings;
-				const FMovieSceneBindingReferences* BindingReferences = Sequence->GetBindingReferences();
-				for (FGuid ObjectGuid : ObjectBindings)
-				{
-					int32 BindingIndex = 0;
-					for (const FMovieSceneBindingReference& Reference : BindingReferences->GetReferences(ObjectGuid))
-					{
-						Bindings.Add({ Reference.ID, BindingIndex++ });
-					}
-				}
-
-				AddChangeBindingTypeMenu(MenuBuilder, Sequencer.ToSharedRef(), Bindings, true, TFunction<void()>());
-			}));
-}
-
 void ULevelSequenceEditorSubsystem::AddBindingPropertiesSidebar(FMenuBuilder& MenuBuilder)
 {
 	AddBindingPropertiesMenu(MenuBuilder);
@@ -1613,21 +1485,74 @@ void ULevelSequenceEditorSubsystem::FBindingPropertiesNotifyHook::NotifyPostChan
 	GEditor->EndTransaction();
 }
 
+TSharedRef<SWidget> ULevelSequenceEditorSubsystem::BuildBindingPropertiesDetailsWidget(TSharedRef<ISequencer> Sequencer)
+{
+	TArray<FGuid> ObjectBindings;
+	Sequencer->GetSelectedObjects(ObjectBindings);
+	if (ObjectBindings.IsEmpty())
+	{
+		return SNullWidget::NullWidget;
+	}
+
+	UMovieSceneSequence* Sequence = Sequencer->GetFocusedMovieSceneSequence();
+	if (!Sequence)
+	{
+		return SNullWidget::NullWidget;
+	}
+
+	FMovieSceneBindingReferences* BindingReferences = Sequence->GetBindingReferences();
+	if (!BindingReferences)
+	{
+		return SNullWidget::NullWidget;
+	}
 
-void ULevelSequenceEditorSubsystem::OnFinishedChangingLocators(const FPropertyChangedEvent& PropertyChangedEvent, TSharedRef<IDetailsView> DetailsView, FGuid ObjectBindingID)
+	// Set up a details panel for the list of locators
+	FDetailsViewArgs DetailsViewArgs;
+	{
+		DetailsViewArgs.bAllowSearch = false;
+		DetailsViewArgs.bCustomFilterAreaLocation = true;
+		DetailsViewArgs.bCustomNameAreaLocation = true;
+		DetailsViewArgs.bHideSelectionTip = true;
+		DetailsViewArgs.bLockable = false;
+		DetailsViewArgs.bSearchInitialKeyFocus = true;
+		DetailsViewArgs.bUpdatesFromSelection = false;
+		DetailsViewArgs.bShowOptions = false;
+		DetailsViewArgs.bShowModifiedPropertiesOption = false;
+		DetailsViewArgs.bShowScrollBar = false;
+		DetailsViewArgs.bAllowMultipleTopLevelObjects = false;
+
+		NotifyHook = FBindingPropertiesNotifyHook(Sequence);
+		DetailsViewArgs.NotifyHook = &NotifyHook;
+
+		DetailsViewArgs.DefaultsOnlyVisibility = EEditDefaultsOnlyNodeVisibility::Automatic;
+	}
+
+	TSharedRef<IDetailsView> DetailsView = FModuleManager::GetModuleChecked<FPropertyEditorModule>("PropertyEditor").CreateDetailView(DetailsViewArgs);
+
+	AddBindingDetailCustomizations(DetailsView, Sequencer, ObjectBindings[0]);
+
+	RefreshBindingDetails(Sequencer, &DetailsView.Get(), ObjectBindings[0]);
+	DetailsView->OnFinishedChangingProperties().AddUObject(this, &ULevelSequenceEditorSubsystem::OnFinishedChangingLocators, Sequencer.ToWeakPtr(), DetailsView, ObjectBindings[0]);
+
+	return DetailsView;
+}
+
+void ULevelSequenceEditorSubsystem::OnFinishedChangingLocators(const FPropertyChangedEvent& PropertyChangedEvent, TWeakPtr<ISequencer> WeakSequencer, TSharedRef<IDetailsView> DetailsView, FGuid ObjectBindingID)
 {
 	const FScopedTransaction Transaction(LOCTEXT("ChangeBindingProperties", "Change Binding Properties"));
 
-	if (!BindingPropertyInfoList)
+	TSharedPtr<ISequencer> Sequencer = WeakSequencer.Pin();
+	if (!Sequencer)
 	{
 		return;
 	}
 
-	TSharedPtr<ISequencer> Sequencer = GetActiveSequencer();
-	if (Sequencer == nullptr)
+	UMovieSceneBindingPropertyInfoList* BindingPropertyInfoList = GetBindingPropertyInfoList(Sequencer.ToSharedRef());
+	if (!BindingPropertyInfoList)
 	{
 		return;
 	}
+
 	UMovieSceneSequence* Sequence = Sequencer->GetFocusedMovieSceneSequence();
 	if (!Sequence)
 	{
@@ -1999,14 +1924,28 @@ void ULevelSequenceEditorSubsystem::RebindComponent(const TArray<FMovieSceneBind
 	}
 }
 
+UMovieSceneBindingPropertyInfoList* ULevelSequenceEditorSubsystem::GetBindingPropertyInfoList(TSharedRef<ISequencer> Sequencer)
+{
+	TObjectPtr<UMovieSceneBindingPropertyInfoList>* BindingPropertyInfoListPtr = BindingPropertyInfoLists.FindByPredicate([Sequencer](const TObjectPtr<UMovieSceneBindingPropertyInfoList>& List) { return List->Sequencer == Sequencer; });
+	if (BindingPropertyInfoListPtr)
+	{
+		return *BindingPropertyInfoListPtr;
+	}
+
+	return nullptr;
+}
+
 void ULevelSequenceEditorSubsystem::RefreshBindingDetails(IDetailsView* DetailsView, FGuid ObjectBindingID)
 {
-	if (DetailsView == nullptr)
+	if (TSharedPtr<ISequencer> Sequencer = GetActiveSequencer())
 	{
-		return;
+		RefreshBindingDetails(Sequencer.ToSharedRef(), DetailsView, ObjectBindingID);
 	}
-	TSharedPtr<ISequencer> Sequencer = GetActiveSequencer();
-	if (Sequencer == nullptr)
+}
+
+void ULevelSequenceEditorSubsystem::RefreshBindingDetails(TSharedRef<ISequencer> Sequencer, IDetailsView* DetailsView, FGuid ObjectBindingID)
+{
+	if (DetailsView == nullptr)
 	{
 		return;
 	}
@@ -2016,29 +1955,32 @@ void ULevelSequenceEditorSubsystem::RefreshBindingDetails(IDetailsView* DetailsV
 	{
 		return;
 	}
-	if (FMovieSceneBindingReferences* BindingReferences = Sequence->GetBindingReferences())
+
+	FMovieSceneBindingReferences* BindingReferences = Sequence->GetBindingReferences();
+	if (!BindingReferences)
 	{
-		if (!BindingPropertyInfoList)
-		{
-			BindingPropertyInfoList = NewObject<UMovieSceneBindingPropertyInfoList>(this);
-		}
-		else
-		{
-			BindingPropertyInfoList->Bindings.Empty();
-		}
+		return;
+	}
 
-		Algo::Transform(BindingReferences->GetReferences(ObjectBindingID), BindingPropertyInfoList->Bindings, [this](const FMovieSceneBindingReference& Reference)
+	UMovieSceneBindingPropertyInfoList* BindingPropertyInfoList = GetBindingPropertyInfoList(Sequencer);
+	if (!BindingPropertyInfoList)
+	{
+		return;
+	}
+
+	BindingPropertyInfoList->Bindings.Empty();
+
+	Algo::Transform(BindingReferences->GetReferences(ObjectBindingID), BindingPropertyInfoList->Bindings, [this](const FMovieSceneBindingReference& Reference)
+		{
+			UMovieSceneCustomBinding* CopiedBinding = nullptr;
+			if (Reference.CustomBinding)
 			{
-				UMovieSceneCustomBinding* CopiedBinding = nullptr;
-				if (Reference.CustomBinding)
-				{
-					CopiedBinding = Cast<UMovieSceneCustomBinding>(StaticDuplicateObject(Reference.CustomBinding, this));
-				}
-				return FMovieSceneBindingPropertyInfo{ Reference.Locator, Reference.ResolveFlags, CopiedBinding };
-			});
+				CopiedBinding = Cast<UMovieSceneCustomBinding>(StaticDuplicateObject(Reference.CustomBinding, this));
+			}
+			return FMovieSceneBindingPropertyInfo{ Reference.Locator, Reference.ResolveFlags, CopiedBinding };
+		});
 
-		DetailsView->SetObject(BindingPropertyInfoList.Get(), true);
-	}
+	DetailsView->SetObject(BindingPropertyInfoList, true);
 }
 
 void ULevelSequenceEditorSubsystem::RefreshTrackRowMetadataDetails(IDetailsView* DetailsView)
diff --git a/Engine/Plugins/MovieScene/LevelSequenceEditor/Source/LevelSequenceEditor/Private/Misc/LevelSequenceCustomization.cpp b/Engine/Plugins/MovieScene/LevelSequenceEditor/Source/LevelSequenceEditor/Private/Misc/LevelSequenceCustomization.cpp
--- a/Engine/Plugins/MovieScene/LevelSequenceEditor/Source/LevelSequenceEditor/Private/Misc/LevelSequenceCustomization.cpp
+++ b/Engine/Plugins/MovieScene/LevelSequenceEditor/Source/LevelSequenceEditor/Private/Misc/LevelSequenceCustomization.cpp
@@ -14,6 +14,7 @@
 #include "ISequencerModule.h"
 #include "LevelSequenceDirector.h"
 #include "LevelSequenceEditorCommands.h"
+#include "LevelSequenceEditorSubsystem.h"
 #include "LevelSequenceFBXInterop.h"
 #include "Modules/ModuleManager.h"
 #include "MovieScene.h"
@@ -283,10 +284,14 @@ void FLevelSequenceCustomization::ExtendObjectBindingContextMenu(FMenuBuilder& M
 			}
 		}
 
+		MenuBuilder.AddSubMenu(
+			LOCTEXT("BindingProperties", "Binding Properties"),
+			LOCTEXT("BindingPropertiesTooltip", "Modify the actor and object bindings for this track"),
+			FNewMenuDelegate::CreateRaw(this, &FLevelSequenceCustomization::AddBindingPropertiesSubMenu));
+
 		// Regular possessable
 		if (!bCustomBinding)
 		{
-			// Regular possessable
 			MenuBuilder.BeginSection("Possessable");
 
 			if (IsSelectedBindingRootPossessable())
@@ -346,8 +351,51 @@ void FLevelSequenceCustomization::ExtendObjectBindingContextMenu(FMenuBuilder& M
 	
 	if (bShowConvert)
 	{
-		// We don't add anything here, but the extension will
 		MenuBuilder.BeginSection("ConvertBinding");
+
+		MenuBuilder.AddSubMenu(
+			LOCTEXT("ConvertBindingLabel", "Convert Selected Binding(s) To..."),
+			LOCTEXT("ConvertBindingLabelTooltip", "Convert selected bindings into another binding type"),
+			FNewMenuDelegate::CreateLambda([this](FMenuBuilder& MenuBuilder)
+				{
+					const TSharedPtr<ISequencer> Sequencer = WeakSequencer.Pin();
+					if (!Sequencer)
+					{
+						return;
+					}
+
+					UMovieSceneSequence* const Sequence = Sequencer->GetFocusedMovieSceneSequence();
+					if (!Sequence)
+					{
+						return;
+					}
+
+					const FMovieSceneBindingReferences* BindingReferences = Sequence->GetBindingReferences();
+					if (!BindingReferences)
+					{
+						return;
+					}
+
+					TArray<FGuid> ObjectBindings;
+					Sequencer->GetSelectedObjects(ObjectBindings);
+					if (ObjectBindings.IsEmpty())
+					{
+						return;
+					}
+
+					TArray<FSequencerChangeBindingInfo> Bindings;
+					for (FGuid ObjectGuid : ObjectBindings)
+					{
+						int32 BindingIndex = 0;
+						for (const FMovieSceneBindingReference& Reference : BindingReferences->GetReferences(ObjectGuid))
+						{
+							Bindings.Add({ Reference.ID, BindingIndex++ });
+						}
+					}
+
+					GEditor->GetEditorSubsystem<ULevelSequenceEditorSubsystem>()->AddChangeBindingTypeMenu(MenuBuilder, Sequencer.ToSharedRef(), Bindings, true, TFunction<void()>());
+				}));
+
 		MenuBuilder.EndSection();
 	}
 
@@ -405,6 +453,19 @@ void FLevelSequenceCustomization::AddAssignActorSubMenu(FMenuBuilder& MenuBuilde
 	MenuBuilder.AddWidget(ActorPickerWidget, FText::GetEmpty(), true);
 }
 
+void FLevelSequenceCustomization::AddBindingPropertiesSubMenu(FMenuBuilder& MenuBuilder)
+{
+	if (TSharedPtr<ISequencer> Sequencer = WeakSequencer.Pin())
+	{
+		if (ULevelSequenceEditorSubsystem* SubSystem = GEditor->GetEditorSubsystem<ULevelSequenceEditorSubsystem>())
+		{
+			TSharedRef<SWidget> DetailsView = SubSystem->BuildBindingPropertiesDetailsWidget(Sequencer.ToSharedRef());
+			MenuBuilder.AddMenuSeparator();
+			MenuBuilder.AddWidget(DetailsView, FText::GetEmpty(), true);
+		}
+	}
+}
+
 TSharedRef<SWidget> FLevelSequenceCustomization::BuildActorPicker()
 {
 	TSharedPtr<ISequencer> Sequencer = WeakSequencer.Pin();
@@ -583,7 +644,31 @@ TSharedPtr<FExtender> FLevelSequenceCustomization::CreateObjectBindingSidebarMen
 
 void FLevelSequenceCustomization::ExtendObjectBindingSidebarMenu(FMenuBuilder& MenuBuilder, TSharedPtr<FObjectBindingModel> ObjectBindingModel)
 {
-	ExtendObjectBindingContextMenu(MenuBuilder, ObjectBindingModel);
+	TSharedPtr<ISequencer> Sequencer = WeakSequencer.Pin();
+	if (!Sequencer)
+	{
+		return;
+	}
+
+	UMovieSceneSequence* Sequence = Sequencer->GetFocusedMovieSceneSequence();
+	if (!Sequence)
+	{
+		return;
+	}
+
+	UMovieScene* MovieScene = Sequence->GetMovieScene();
+
+	FGuid ObjectBindingID = ObjectBindingModel->GetObjectGuid();
+
+	if (!MovieScene || !ObjectBindingID.IsValid())
+	{
+		return;
+	}
+
+	if (FMovieScenePossessable* Possessable = MovieScene->FindPossessable(ObjectBindingID))
+	{
+		AddBindingPropertiesSubMenu(MenuBuilder);
+	}
 }
 
 } // namespace UE::Sequencer
diff --git a/Engine/Plugins/MovieScene/LevelSequenceEditor/Source/LevelSequenceEditor/Private/Misc/LevelSequenceCustomization.h b/Engine/Plugins/MovieScene/LevelSequenceEditor/Source/LevelSequenceEditor/Private/Misc/LevelSequenceCustomization.h
--- a/Engine/Plugins/MovieScene/LevelSequenceEditor/Source/LevelSequenceEditor/Private/Misc/LevelSequenceCustomization.h
+++ b/Engine/Plugins/MovieScene/LevelSequenceEditor/Source/LevelSequenceEditor/Private/Misc/LevelSequenceCustomization.h
@@ -48,6 +48,7 @@ private:
 	/** Object Binding Context Menu Extensions */
 	void ExtendObjectBindingContextMenu(FMenuBuilder& MenuBuilder, TSharedPtr<FObjectBindingModel> ObjectBindingModel);
 	void AddAssignActorSubMenu(FMenuBuilder& MenuBuilder);
+	void AddBindingPropertiesSubMenu(FMenuBuilder& MenuBuilder);
 	TSharedRef<SWidget> BuildActorPicker();
 
 	void AddActorsToBinding();
diff --git a/Engine/Plugins/MovieScene/LevelSequenceEditor/Source/LevelSequenceEditor/Public/LevelSequenceEditorSubsystem.h b/Engine/Plugins/MovieScene/LevelSequenceEditor/Source/LevelSequenceEditor/Public/LevelSequenceEditorSubsystem.h
--- a/Engine/Plugins/MovieScene/LevelSequenceEditor/Source/LevelSequenceEditor/Public/LevelSequenceEditorSubsystem.h
+++ b/Engine/Plugins/MovieScene/LevelSequenceEditor/Source/LevelSequenceEditor/Public/LevelSequenceEditorSubsystem.h
@@ -43,6 +43,7 @@ class USequencerCurveEditorObject;
 class UMovieSceneCustomBinding;
 class IMenu;
 class UMovieSceneTrackRowMetadataHelper;
+class SWidget;
 
 USTRUCT(BlueprintType)
 struct FMovieSceneScriptingParams
@@ -82,6 +83,9 @@ public:
 	// List of locator info for a particular binding
 	UPROPERTY(EditAnywhere, Category = "Binding Properties")
 	TArray<FMovieSceneBindingPropertyInfo> Bindings;
+
+	// The Sequencer that the Bindings are sourced from
+	TWeakPtr<ISequencer> Sequencer;
 };
 
 /**
@@ -293,7 +297,7 @@ private:
 	};
 
 	UPROPERTY()
-	TObjectPtr<UMovieSceneBindingPropertyInfoList> BindingPropertyInfoList = nullptr;
+	TArray<TObjectPtr<UMovieSceneBindingPropertyInfoList>> BindingPropertyInfoLists;
 
 	UPROPERTY()
 	TArray<TObjectPtr<UMovieSceneTrackRowMetadataHelper>> TrackRowMetadataHelperList;
@@ -304,7 +308,6 @@ private:
 
 	UE_API void AddBindingDetailCustomizations(TSharedRef<IDetailsView> DetailsView, TSharedPtr<ISequencer> ActiveSequencer, FGuid BindingGuid);
 	UE_API void AddTrackRowMetadataCustomizations(TSharedRef<IDetailsView> DetailsView, TSharedPtr<ISequencer> ActiveSequencer, UMovieSceneSequence* Sequence);
-	UE_API void OnBindingPropertyMenuBeingDestroyed(const TSharedRef<IMenu>& Menu, TSharedRef<IDetailsView> DetailsView);
 	UE_API void OnTrackRowMetadataMenuBeingDestroyed(const TSharedRef<IMenu>& Menu, TSharedRef<IDetailsView> DetailsView);
 
 	UE_API TSharedPtr<ISequencer> GetActiveSequencer();
@@ -312,8 +315,6 @@ private:
 	UE_API void RebindComponentInternal(const FName& ComponentName);
 
 	UE_API void AddBindingPropertiesMenu(FMenuBuilder& MenuBuilder);
-	UE_API void AddConvertBindingsMenu(FMenuBuilder& MenuBuilder);
-
 	
 	UE_API void FillDirectorBlueprintBindingSubMenu(FMenuBuilder& MenuBuilder, TSharedRef<ISequencer> Sequencer, const TArray<FSequencerChangeBindingInfo>& BindingsToChange, bool bConvert, TFunction<void()> OnBindingChanged, const TSubclassOf<UMovieSceneCustomBinding>& CustomBindingType);
 	UE_API void PopulateQuickBindSubMenu(FMenuBuilder& MenuBuilder, TSharedRef<ISequencer> Sequencer, const TArray<FSequencerChangeBindingInfo>& BindingsToChange, bool bConvert, TFunction<void()> OnBindingChanged, const TSubclassOf<UMovieSceneCustomBinding>& CustomBindingType);
@@ -325,20 +326,25 @@ private:
 
 	UE_API void AddTrackRowMetadataMenu(FMenuBuilder& MenuBuilder);
 public:
-	UE_API void AddBindingPropertiesSidebar(FMenuBuilder& MenuBuilder); 
-	
+	UE_API void AddBindingPropertiesSidebar(FMenuBuilder& MenuBuilder);
+
+	UE_API TSharedRef<SWidget> BuildBindingPropertiesDetailsWidget(TSharedRef<ISequencer> Sequencer);
+	UE_API void RefreshBindingDetails(TSharedRef<ISequencer> Sequencer, IDetailsView* DetailsView, FGuid ObjectBindingID);
+
 	/* Creates a menu for changing or converting a binding type. If bConvert is true, it will only show types that state they are able to be converted to from the passed in bindings
 	and will attempt to convert them. If bConvert is false, it will change the binding type and reset to a default binding of that type.*/
 	UE_API void AddChangeBindingTypeMenu(FMenuBuilder& MenuBuilder, TSharedRef<ISequencer> Sequencer, const TArray<FSequencerChangeBindingInfo>& BindingsToChange, bool bConvert, TFunction<void()> OnBindingChanged);
 
 private:
-	UE_API void OnFinishedChangingLocators(const FPropertyChangedEvent& PropertyChangedEvent, TSharedRef<IDetailsView> DetailsView, FGuid ObjectBindingID);
+	void OnFinishedChangingLocators(const FPropertyChangedEvent& PropertyChangedEvent, TWeakPtr<ISequencer> WeakSequencer, TSharedRef<IDetailsView> DetailsView, FGuid ObjectBindingID);
 
 	UE_API void OnFinishedChangingTrackRowMetadata(const FPropertyChangedEvent& PropertyChangedEvent, TSharedRef<IDetailsView> DetailsView);
 
 	UE_API void GetRebindComponentNames(TArray<FName>& OutComponentNames);
 	UE_API void RebindComponentMenu(FMenuBuilder& MenuBuilder);
 
+	UMovieSceneBindingPropertyInfoList* GetBindingPropertyInfoList(TSharedRef<ISequencer> Sequencer);
+
 	FDelegateHandle OnSequencerCreatedHandle;
 
 	/* List of sequencers that have been created */
@@ -350,7 +356,6 @@ private:
 	UPROPERTY()
 	TArray<TObjectPtr<USequencerCurveEditorObject>> CurveEditorArray;
 
-	TSharedPtr<FExtender> BindingPropertiesMenuExtender;
 	TSharedPtr<FExtender> RebindComponentMenuExtender;
 	TSharedPtr<FExtender> SidebarMenuExtender;
 
diff --git a/Engine/Source/Editor/Sequencer/Private/SequencerCommonHelpers.cpp b/Engine/Source/Editor/Sequencer/Private/SequencerCommonHelpers.cpp
--- a/Engine/Source/Editor/Sequencer/Private/SequencerCommonHelpers.cpp
+++ b/Engine/Source/Editor/Sequencer/Private/SequencerCommonHelpers.cpp
@@ -627,6 +627,41 @@ void SequencerHelpers::BuildBlendingMenu(const TWeakPtr<FSequencer>& InWeakSeque
 		}));
 }
 
+void SequencerHelpers::BuildObjectBindingMenu(const TWeakPtr<FSequencer>& InWeakSequencer
+	, const UE::Sequencer::TViewModelPtr<UE::Sequencer::FObjectBindingModel>& InObjectBindingModel
+	, FMenuBuilder& MenuBuilder)
+{
+	TSharedPtr<FSequencer> Sequencer = InWeakSequencer.Pin();
+	if (!Sequencer || !InObjectBindingModel)
+	{
+		return;
+	}
+
+	const TSharedPtr<UE::Sequencer::FSequencerEditorViewModel> EditorViewModel = InObjectBindingModel->GetEditor();
+	if (!EditorViewModel)
+	{
+		return;
+	}
+
+	ISequencerModule& SequencerModule = FModuleManager::GetModuleChecked<ISequencerModule>(TEXT("Sequencer"));
+
+	TArray<UObject*> ContextObjects = { Sequencer->FindSpawnedObjectOrTemplate(InObjectBindingModel->GetObjectGuid()) };
+
+	const TSharedPtr<FExtender> Extender = EditorViewModel->GetSequencerMenuExtender(
+		SequencerModule.GetSidebarExtensibilityManager(), 
+		ContextObjects,
+		&FSequencerCustomizationInfo::OnBuildSidebarMenu, 
+		InObjectBindingModel);
+
+	if (Extender.IsValid())
+	{
+		MenuBuilder.PushExtender(Extender.ToSharedRef());
+	}
+
+	MenuBuilder.BeginSection(TEXT("ObjectBindingActions"));
+	MenuBuilder.EndSection();
+}
+
 TArray<TWeakObjectPtr<>> SequencerHelpers::GetSectionObjectsFromTrackAreaModels(const UE::Sequencer::FViewModelVariantIterator& InTrackAreaModels)
 {
 	using namespace UE::Sequencer;
diff --git a/Engine/Source/Editor/Sequencer/Private/SequencerCommonHelpers.h b/Engine/Source/Editor/Sequencer/Private/SequencerCommonHelpers.h
--- a/Engine/Source/Editor/Sequencer/Private/SequencerCommonHelpers.h
+++ b/Engine/Source/Editor/Sequencer/Private/SequencerCommonHelpers.h
@@ -22,6 +22,7 @@ namespace UE
 namespace Sequencer
 {
 	class FChannelModel;
+	class FObjectBindingModel;
 	class IOutlinerExtension;
 	struct FKeySelection;
 	struct FViewModelVariantIterator;
@@ -109,6 +110,13 @@ public:
 		, const TWeakObjectPtr<UMovieSceneTrack>& InTrackWeak
 		, FMenuBuilder& MenuBuilder);
 
+	/*
+	 * Build a menu for selected object bindings
+	 */
+	static void BuildObjectBindingMenu(const TWeakPtr<FSequencer>& InWeakSequencer
+		, const UE::Sequencer::TViewModelPtr<UE::Sequencer::FObjectBindingModel>& InObjectBindingModel
+		, FMenuBuilder& MenuBuilder);
+
 	/**
 	 * Gets all section objects from track area models
 	 */
diff --git a/Engine/Source/Editor/Sequencer/Private/Widgets/Sidebar/SequencerSelectionDrawer.cpp b/Engine/Source/Editor/Sequencer/Private/Widgets/Sidebar/SequencerSelectionDrawer.cpp
--- a/Engine/Source/Editor/Sequencer/Private/Widgets/Sidebar/SequencerSelectionDrawer.cpp
+++ b/Engine/Source/Editor/Sequencer/Private/Widgets/Sidebar/SequencerSelectionDrawer.cpp
@@ -362,6 +362,7 @@ void FSequencerSelectionDrawer::BuildOutlinerDetails(const TSharedRef<FSequencer
 		else if (const TViewModelPtr<FObjectBindingModel> ObjectBindingModel = OutlinerItem.ImplicitCast())
 		{
 			ObjectBindings.Add(ObjectBindingModel);
+			SequencerHelpers::BuildObjectBindingMenu(WeakSequencer, ObjectBindingModel, MenuBuilder);
 		}
 		else if (const TViewModelPtr<FFolderModel> FolderModel = OutlinerItem.ImplicitCast())
 		{
