Sequencer: Fix jitter in the viewport when the aspect ratio is animated

Seems to have come in from 2875029

#Synced-CL 48964464

diff --git a/Engine/Plugins/MovieScene/LevelSequenceEditor/Source/LevelSequenceEditor/Private/CinematicViewport/SCinematicLevelViewport.cpp b/Engine/Plugins/MovieScene/LevelSequenceEditor/Source/LevelSequenceEditor/Private/CinematicViewport/SCinematicLevelViewport.cpp
--- a/Engine/Plugins/MovieScene/LevelSequenceEditor/Source/LevelSequenceEditor/Private/CinematicViewport/SCinematicLevelViewport.cpp
+++ b/Engine/Plugins/MovieScene/LevelSequenceEditor/Source/LevelSequenceEditor/Private/CinematicViewport/SCinematicLevelViewport.cpp
@@ -719,10 +719,10 @@ void SCinematicLevelViewport::CacheDesiredViewportSize(const FGeometry& Allotted
 	FVector2D AllowableSpace = AllottedGeometry.GetLocalSize();
 	AllowableSpace.Y -= ViewportControls->GetDesiredSize().Y;
 
-	if (ViewportClient->IsAspectRatioConstrained())
+	if (ViewportClient->IsAspectRatioConstrained() && ViewportClient->AspectRatio > 0)
 	{
-		const float MinSize = FMath::TruncToFloat(FMath::Min(AllowableSpace.X / ViewportClient->AspectRatio, AllowableSpace.Y));
-		DesiredViewportSize = FVector2D(FMath::TruncToFloat(ViewportClient->AspectRatio * MinSize), MinSize);
+		const double MinSize = FMath::Min(AllowableSpace.X / (double)ViewportClient->AspectRatio, AllowableSpace.Y);
+		DesiredViewportSize = FVector2D((double)ViewportClient->AspectRatio * MinSize, MinSize);
 	}
 	else
 	{
