Sequencer: Use the actor component name if not an actor before using the class name.

This fixes a minor bug where a possessable binding will be renamed to the class name if it's a component, rather than keeping it with the component name.

#Synced-CL 50405187

diff --git a/Engine/Plugins/MovieScene/LevelSequenceEditor/Source/LevelSequenceEditor/Private/LevelSequenceEditorSubsystem.cpp b/Engine/Plugins/MovieScene/LevelSequenceEditor/Source/LevelSequenceEditor/Private/LevelSequenceEditorSubsystem.cpp
--- a/Engine/Plugins/MovieScene/LevelSequenceEditor/Source/LevelSequenceEditor/Private/LevelSequenceEditorSubsystem.cpp
+++ b/Engine/Plugins/MovieScene/LevelSequenceEditor/Source/LevelSequenceEditor/Private/LevelSequenceEditorSubsystem.cpp
@@ -1633,6 +1633,10 @@ void ULevelSequenceEditorSubsystem::OnFinishedChangingLocators(const FPropertyCh
 				{
 					Possessable->SetName(Actor->GetActorLabel());
 				}
+				else if (UActorComponent* ActorComponent = Cast<UActorComponent>(ObjectsInCurrentSequence[0].Get()))
+				{
+					Possessable->SetName(ActorComponent->GetName());
+				}
 				else
 				{
 					Possessable->SetName(ObjectClass->GetName());
