Sequencer: Add Bake Transform to object binding menu

Moved the command to Level Sequence Customization (similar to Import/Export FBX) since the object binding menu extension is there

Tested that Bake Transform still works for Camera Animation and Template Sequences

#Synced-CL 50973831

diff --git a/Engine/Plugins/MovieScene/LevelSequenceEditor/Source/LevelSequenceEditor/Private/LevelSequenceEditorCommands.cpp b/Engine/Plugins/MovieScene/LevelSequenceEditor/Source/LevelSequenceEditor/Private/LevelSequenceEditorCommands.cpp
--- a/Engine/Plugins/MovieScene/LevelSequenceEditor/Source/LevelSequenceEditor/Private/LevelSequenceEditorCommands.cpp
+++ b/Engine/Plugins/MovieScene/LevelSequenceEditor/Source/LevelSequenceEditor/Private/LevelSequenceEditorCommands.cpp
@@ -42,6 +42,8 @@ void FLevelSequenceEditorCommands::RegisterCommands()
 
 	UI_COMMAND(ImportFBX, "Import...", "Import the animation from an FBX file.", EUserInterfaceActionType::Button, FInputChord());
 	UI_COMMAND(ExportFBX, "Export...", "Export the selected objects (all if none selected) and animation to an FBX file. (Shots and sub-scenes not supported)", EUserInterfaceActionType::Button, FInputChord());
+
+	UI_COMMAND(BakeTransform, "Bake Transform", "Bake transform in world space, removing any existing transform and attach tracks", EUserInterfaceActionType::Button, FInputChord());
 }
 
 #undef LOCTEXT_NAMESPACE
diff --git a/Engine/Plugins/MovieScene/LevelSequenceEditor/Source/LevelSequenceEditor/Private/LevelSequenceEditorCommands.h b/Engine/Plugins/MovieScene/LevelSequenceEditor/Source/LevelSequenceEditor/Private/LevelSequenceEditorCommands.h
--- a/Engine/Plugins/MovieScene/LevelSequenceEditor/Source/LevelSequenceEditor/Private/LevelSequenceEditorCommands.h
+++ b/Engine/Plugins/MovieScene/LevelSequenceEditor/Source/LevelSequenceEditor/Private/LevelSequenceEditorCommands.h
@@ -36,4 +36,5 @@ public:
 	/** Exports animation to fbx. */
 	TSharedPtr< FUICommandInfo > ExportFBX;
 
+	TSharedPtr<FUICommandInfo> BakeTransform;
 };
diff --git a/Engine/Plugins/MovieScene/LevelSequenceEditor/Source/LevelSequenceEditor/Private/Misc/LevelSequenceCustomization.cpp b/Engine/Plugins/MovieScene/LevelSequenceEditor/Source/LevelSequenceEditor/Private/Misc/LevelSequenceCustomization.cpp
--- a/Engine/Plugins/MovieScene/LevelSequenceEditor/Source/LevelSequenceEditor/Private/Misc/LevelSequenceCustomization.cpp
+++ b/Engine/Plugins/MovieScene/LevelSequenceEditor/Source/LevelSequenceEditor/Private/Misc/LevelSequenceCustomization.cpp
@@ -3,6 +3,7 @@
 #include "Misc/LevelSequenceCustomization.h"
 
 #include "ActorTreeItem.h"
+#include "BakingAnimationKeySettings.h"
 #include "Bindings/MovieSceneSpawnableBinding.h"
 #include "ClassViewerFilter.h"
 #include "ClassViewerModule.h"
@@ -30,6 +31,7 @@
 #include "SequencerCommands.h"
 #include "SequencerSettings.h"
 #include "SequencerUtilities.h"
+#include "Widgets/SLevelSequenceBakeTransform.h"
 
 #define LOCTEXT_NAMESPACE "LevelSequenceCustomization"
 
@@ -102,6 +104,10 @@ void FLevelSequenceCustomization::BindCommands()
 		Commands.FixActorReferences,
 		FExecuteAction::CreateRaw(this, &FLevelSequenceCustomization::FixActorReferences));
 
+	CommandList->MapAction(
+		Commands.BakeTransform,
+		FExecuteAction::CreateRaw(this, &FLevelSequenceCustomization::BakeTransform));
+
 	CommandList->MapAction(
 		Commands.AddActorsToBinding,
 		FExecuteAction::CreateRaw(this, &FLevelSequenceCustomization::AddActorsToBinding),
@@ -134,6 +140,7 @@ TSharedPtr<FExtender> FLevelSequenceCustomization::CreateActionsMenuExtender()
 		"SequenceOptions", EExtensionHook::First, CommandList,
 		FMenuExtensionDelegate::CreateLambda([](FMenuBuilder& MenuBuilder)
 			{
+				MenuBuilder.AddMenuEntry(FLevelSequenceEditorCommands::Get().BakeTransform);
 				MenuBuilder.AddMenuEntry(FLevelSequenceEditorCommands::Get().ImportFBX);
 				MenuBuilder.AddMenuEntry(FLevelSequenceEditorCommands::Get().ExportFBX);
 			}));
@@ -210,6 +217,51 @@ void FLevelSequenceCustomization::FixActorReferences()
 	}
 }
 
+void FLevelSequenceCustomization::BakeTransform()
+{
+	TSharedPtr<ISequencer> Sequencer = WeakSequencer.Pin();
+	if (!Sequencer)
+	{
+		return;
+	}
+
+	UMovieSceneSequence* Sequence = Sequencer->GetFocusedMovieSceneSequence();
+	UMovieScene* MovieScene = Sequence ? Sequence->GetMovieScene() : nullptr;
+	if (!MovieScene)
+	{
+		return;
+	}
+
+	TArray<FGuid> ObjectBindings;
+	Sequencer->GetSelectedObjects(ObjectBindings);
+
+	TArray<FMovieSceneBindingProxy> BindingProxies;
+	for (FGuid Guid : ObjectBindings)
+	{
+		BindingProxies.Add(FMovieSceneBindingProxy(Guid, Sequence));
+	}
+
+	static FBakingAnimationKeySettings Settings; //reuse the settings except for the range
+	Settings.StartFrame = UE::MovieScene::DiscreteInclusiveLower(MovieScene->GetPlaybackRange());
+	Settings.EndFrame = UE::MovieScene::DiscreteExclusiveUpper(MovieScene->GetPlaybackRange());
+
+	TSharedRef<SLevelSequenceBakeTransform> BakeWidget =
+		SNew(SLevelSequenceBakeTransform)
+		.Settings(Settings)
+		.Sequencer(Sequencer.Get())
+		.OnBake_Lambda([this, &BindingProxies, Sequencer](FBakingAnimationKeySettings InSettings)
+			{
+				if (FSequencerUtilities::BakeTransform(Sequencer.ToSharedRef(), BindingProxies, InSettings))
+				{
+					Settings = InSettings;
+				}
+
+				return FReply::Handled();
+			});
+
+	BakeWidget->OpenDialog(true);
+}
+
 TSharedPtr<FExtender> FLevelSequenceCustomization::CreateObjectBindingContextMenuExtender(FViewModelPtr InViewModel)
 {
 	TSharedRef<FExtender> Extender = MakeShared<FExtender>();
@@ -401,6 +453,8 @@ void FLevelSequenceCustomization::ExtendObjectBindingContextMenu(FMenuBuilder& M
 
 	MenuBuilder.BeginSection("Import/Export", LOCTEXT("ImportExportMenuSectionName", "Import/Export"));
 
+	MenuBuilder.AddMenuEntry(FLevelSequenceEditorCommands::Get().BakeTransform);
+
 	MenuBuilder.AddMenuEntry(
 		LOCTEXT("ImportFBX", "Import..."),
 		LOCTEXT("ImportFBXTooltip", "Import FBX animation to this object"),
diff --git a/Engine/Plugins/MovieScene/LevelSequenceEditor/Source/LevelSequenceEditor/Private/Misc/LevelSequenceCustomization.h b/Engine/Plugins/MovieScene/LevelSequenceEditor/Source/LevelSequenceEditor/Private/Misc/LevelSequenceCustomization.h
--- a/Engine/Plugins/MovieScene/LevelSequenceEditor/Source/LevelSequenceEditor/Private/Misc/LevelSequenceCustomization.h
+++ b/Engine/Plugins/MovieScene/LevelSequenceEditor/Source/LevelSequenceEditor/Private/Misc/LevelSequenceCustomization.h
@@ -42,7 +42,8 @@ private:
 	void SnapSectionsToTimelineUsingSourceTimecode();
 	void SyncSectionsUsingSourceTimecode();
 	void FixActorReferences();
-
+	void BakeTransform();
+	
 	bool DoesMovieSceneHaveEnoughSectionsSelected(int32 MinSectionCount);
 
 	/** Object Binding Context Menu Extensions */
diff --git a/Engine/Plugins/MovieScene/LevelSequenceEditor/Source/LevelSequenceEditor/Private/Widgets/SLevelSequenceBakeTransform.cpp b/Engine/Plugins/MovieScene/LevelSequenceEditor/Source/LevelSequenceEditor/Private/Widgets/SLevelSequenceBakeTransform.cpp
new file mode 100644
--- /dev/null
+++ b/Engine/Plugins/MovieScene/LevelSequenceEditor/Source/LevelSequenceEditor/Private/Widgets/SLevelSequenceBakeTransform.cpp
@@ -0,0 +1,162 @@
+// Copyright Epic Games, Inc. All Rights Reserved.
+
+#include "Widgets/SLevelSequenceBakeTransform.h"
+
+#include "Editor.h"
+#include "Framework/Application/SlateApplication.h"
+#include "ISequencer.h"
+#include "IStructureDetailsView.h"
+#include "PropertyEditorModule.h"
+#include "Widgets/SWindow.h"
+#include "Widgets/Input/SButton.h"
+#include "Widgets/Layout/SSpacer.h"
+
+#define LOCTEXT_NAMESPACE "SLevelSequenceBakeTransform"
+
+void SLevelSequenceBakeTransform::Construct(const FArguments& InArgs)
+{
+	check(InArgs._Sequencer);
+	check(InArgs._OnBake.IsBound());
+
+	Settings = MakeShared<TStructOnScope<FBakingAnimationKeySettings>>();
+	Settings->InitializeAs<FBakingAnimationKeySettings>();
+	*Settings = InArgs._Settings;
+	//always setting space to be parent as default, since stored space may not be available.
+	Sequencer = InArgs._Sequencer;
+
+	FStructureDetailsViewArgs StructureViewArgs;
+	StructureViewArgs.bShowObjects = true;
+	StructureViewArgs.bShowAssets = true;
+	StructureViewArgs.bShowClasses = true;
+	StructureViewArgs.bShowInterfaces = true;
+
+	FDetailsViewArgs ViewArgs;
+	ViewArgs.bAllowSearch = false;
+	ViewArgs.bHideSelectionTip = false;
+	ViewArgs.bShowObjectLabel = false;
+
+	FPropertyEditorModule& PropertyEditor = FModuleManager::Get().LoadModuleChecked<FPropertyEditorModule>(TEXT("PropertyEditor"));
+
+	DetailsView = PropertyEditor.CreateStructureDetailView(ViewArgs, StructureViewArgs, TSharedPtr<FStructOnScope>());
+
+	DetailsView->GetDetailsView()->RegisterInstancedCustomPropertyTypeLayout("FrameNumber",
+		FOnGetPropertyTypeCustomizationInstance::CreateSP(Sequencer, &ISequencer::MakeFrameNumberDetailsCustomization));
+	DetailsView->SetStructureData(Settings);
+
+	ChildSlot
+		[
+			SNew(SBorder)
+				.Visibility(EVisibility::Visible)
+				[
+					SNew(SVerticalBox)
+
+
+						+ SVerticalBox::Slot()
+						.AutoHeight()
+						.Padding(0.f, 0.f, 0.f, 0.f)
+						[
+							DetailsView->GetWidget().ToSharedRef()
+						]
+
+						+ SVerticalBox::Slot()
+						.AutoHeight()
+						.Padding(0.f, 16.f, 0.f, 16.f)
+						[
+							SNew(SHorizontalBox)
+
+								+ SHorizontalBox::Slot()
+								.FillWidth(1.f)
+								[
+									SNew(SSpacer)
+								]
+
+								+ SHorizontalBox::Slot()
+								.AutoWidth()
+								.HAlign(HAlign_Left)
+								.VAlign(VAlign_Center)
+								.Padding(8.f, 0.f, 0.f, 0.f)
+								[
+									SNew(SButton)
+										.HAlign(HAlign_Center)
+										.ContentPadding(FAppStyle::GetMargin("StandardDialog.ContentPadding"))
+										.Text(LOCTEXT("OK", "OK"))
+										.OnClicked_Lambda([this, InArgs]()
+											{
+												FReply Reply = InArgs._OnBake.Execute(*(Settings->Get()));
+												CloseDialog();
+												return Reply;
+
+											})
+								]
+
+							+ SHorizontalBox::Slot()
+								.AutoWidth()
+								.HAlign(HAlign_Left)
+								.VAlign(VAlign_Center)
+								.Padding(8.f, 0.f, 16.f, 0.f)
+								[
+									SNew(SButton)
+										.HAlign(HAlign_Center)
+										.ContentPadding(FAppStyle::GetMargin("StandardDialog.ContentPadding"))
+										.Text(LOCTEXT("Cancel", "Cancel"))
+										.OnClicked_Lambda([this]()
+											{
+												CloseDialog();
+												return FReply::Handled();
+											})
+								]
+						]
+				]
+		];
+}
+
+class SLevelSequenceBakeTransformDialogWindow : public SWindow
+{
+};
+
+FReply SLevelSequenceBakeTransform::OpenDialog(bool bModal)
+{
+	check(!DialogWindow.IsValid());
+
+	const FVector2D CursorPos = FSlateApplication::Get().GetCursorPos();
+
+	TSharedRef<SLevelSequenceBakeTransformDialogWindow> Window = SNew(SLevelSequenceBakeTransformDialogWindow)
+		.Title(LOCTEXT("SLevelSequenceBakeTransformWidgetTitle", "Bake Transform"))
+		.CreateTitleBar(true)
+		.Type(EWindowType::Normal)
+		.SizingRule(ESizingRule::Autosized)
+		.ScreenPosition(CursorPos)
+		.FocusWhenFirstShown(true)
+		.ActivationPolicy(EWindowActivationPolicy::FirstShown)
+		[
+			AsShared()
+		];
+
+	Window->SetWidgetToFocusOnActivate(AsShared());
+
+	DialogWindow = Window;
+
+	Window->MoveWindowTo(CursorPos);
+
+	if (bModal)
+	{
+		GEditor->EditorAddModalWindow(Window);
+	}
+	else
+	{
+		FSlateApplication::Get().AddWindow(Window);
+	}
+
+	return FReply::Handled();
+}
+
+void SLevelSequenceBakeTransform::CloseDialog()
+{
+	if (DialogWindow.IsValid())
+	{
+		DialogWindow.Pin()->RequestDestroyWindow();
+		DialogWindow.Reset();
+	}
+}
+
+#undef LOCTEXT_NAMESPACE
\ No newline at end of file
diff --git a/Engine/Plugins/MovieScene/LevelSequenceEditor/Source/LevelSequenceEditor/Public/Widgets/SLevelSequenceBakeTransform.h b/Engine/Plugins/MovieScene/LevelSequenceEditor/Source/LevelSequenceEditor/Public/Widgets/SLevelSequenceBakeTransform.h
new file mode 100644
--- /dev/null
+++ b/Engine/Plugins/MovieScene/LevelSequenceEditor/Source/LevelSequenceEditor/Public/Widgets/SLevelSequenceBakeTransform.h
@@ -0,0 +1,49 @@
+// Copyright Epic Games, Inc. All Rights Reserved.
+
+#pragma once
+
+#include "BakingAnimationKeySettings.h"
+#include "Widgets/DeclarativeSyntaxSupport.h"
+#include "Widgets/SCompoundWidget.h"
+#include "UObject/StructOnScope.h"
+
+class ISequencer;
+class IStructureDetailsView;
+class SWindow;
+
+#define UE_API LEVELSEQUENCEEDITOR_API
+
+DECLARE_DELEGATE_RetVal_OneParam(FReply, SLevelSequenceBakeTransformOnBake, FBakingAnimationKeySettings);
+
+/** Widget that displays bake options for baking transforms */
+class SLevelSequenceBakeTransform : public SCompoundWidget
+{
+public:
+
+	SLATE_BEGIN_ARGS(SLevelSequenceBakeTransform)
+		: _Sequencer(nullptr)
+		{
+		}
+		SLATE_ARGUMENT(ISequencer*, Sequencer)
+		SLATE_ARGUMENT(FBakingAnimationKeySettings, Settings)
+		SLATE_EVENT(SLevelSequenceBakeTransformOnBake, OnBake)
+	SLATE_END_ARGS()
+
+	UE_API void Construct(const FArguments& InArgs);
+	virtual ~SLevelSequenceBakeTransform() override = default;
+
+	UE_API FReply OpenDialog(bool bModal = true);
+	void CloseDialog();
+
+private:
+
+	//used for setting up the details
+	TSharedPtr<TStructOnScope<FBakingAnimationKeySettings>> Settings;
+
+	ISequencer* Sequencer;
+
+	TWeakPtr<SWindow> DialogWindow;
+	TSharedPtr<IStructureDetailsView> DetailsView;
+};
+
+#undef UE_API
diff --git a/Engine/Plugins/MovieScene/TemplateSequence/Source/TemplateSequenceEditor/Private/Customizations/TemplateSequenceCustomizationBase.cpp b/Engine/Plugins/MovieScene/TemplateSequence/Source/TemplateSequenceEditor/Private/Customizations/TemplateSequenceCustomizationBase.cpp
--- a/Engine/Plugins/MovieScene/TemplateSequence/Source/TemplateSequenceEditor/Private/Customizations/TemplateSequenceCustomizationBase.cpp
+++ b/Engine/Plugins/MovieScene/TemplateSequence/Source/TemplateSequenceEditor/Private/Customizations/TemplateSequenceCustomizationBase.cpp
@@ -1,6 +1,8 @@
 // Copyright Epic Games, Inc. All Rights Reserved.
 
 #include "TemplateSequenceCustomizationBase.h"
+
+#include "BakingAnimationKeySettings.h"
 #include "Framework/Application/SlateApplication.h"
 #include "Framework/MultiBox/MultiBoxBuilder.h"
 #include "Framework/MultiBox/MultiBoxExtender.h"
@@ -8,8 +10,12 @@
 #include "ISequencer.h"
 #include "LevelSequenceFBXInterop.h"
 #include "Misc/TemplateSequenceEditorUtil.h"
+#include "MovieScene.h"
+#include "MovieSceneTimeHelpers.h"
 #include "ScopedTransaction.h"
+#include "SequencerUtilities.h"
 #include "TemplateSequence.h"
+#include "Widgets/SLevelSequenceBakeTransform.h"
 
 #define LOCTEXT_NAMESPACE "TemplateSequenceCustomizationBase"
 
@@ -101,6 +107,56 @@ void FTemplateSequenceCustomizationBase::ExtendObjectBindingContextMenu(FMenuBui
 
 	MenuBuilder.BeginSection("Import/Export", LOCTEXT("ImportExportMenuSectionName", "Import/Export"));
 
+	MenuBuilder.AddMenuEntry(
+		LOCTEXT("BakeTransform", "Bake Transform"),
+		LOCTEXT("BakeTransformTooltip", "Bake transform in world space, removing any existing transform and attach tracks"),
+		FSlateIcon(),
+		FUIAction(
+			FExecuteAction::CreateLambda([this] {
+				const TSharedPtr<ISequencer> Sequencer = WeakSequencer.Pin();
+				if (!Sequencer.IsValid())
+				{
+					return;
+				}
+
+				UMovieSceneSequence* Sequence = Sequencer->GetFocusedMovieSceneSequence();
+				UMovieScene* MovieScene = Sequence ? Sequence->GetMovieScene() : nullptr;
+				if (!MovieScene)
+				{
+					return;
+				}
+
+				TArray<FGuid> ObjectBindings;
+				Sequencer->GetSelectedObjects(ObjectBindings);
+
+				TArray<FMovieSceneBindingProxy> BindingProxies;
+				for (FGuid Guid : ObjectBindings)
+				{
+					BindingProxies.Add(FMovieSceneBindingProxy(Guid, Sequence));
+				}
+
+				static FBakingAnimationKeySettings Settings; //reuse the settings except for the range
+				Settings.StartFrame = UE::MovieScene::DiscreteInclusiveLower(MovieScene->GetPlaybackRange());
+				Settings.EndFrame = UE::MovieScene::DiscreteExclusiveUpper(MovieScene->GetPlaybackRange());
+
+				TSharedRef<SLevelSequenceBakeTransform> BakeWidget =
+					SNew(SLevelSequenceBakeTransform)
+					.Settings(Settings)
+					.Sequencer(Sequencer.Get())
+					.OnBake_Lambda([this, &BindingProxies, Sequencer](FBakingAnimationKeySettings InSettings)
+						{
+							if (FSequencerUtilities::BakeTransform(Sequencer.ToSharedRef(), BindingProxies, InSettings))
+							{
+								Settings = InSettings;
+							}
+
+							return FReply::Handled();
+						});
+
+				BakeWidget->OpenDialog(true);
+			})
+		));
+
 	MenuBuilder.AddMenuEntry(
 		LOCTEXT("ImportFBX", "Import..."),
 		LOCTEXT("ImportFBXTooltip", "Import FBX animation to this object"),
diff --git a/Engine/Source/Editor/Sequencer/Private/SSequencer.cpp b/Engine/Source/Editor/Sequencer/Private/SSequencer.cpp
--- a/Engine/Source/Editor/Sequencer/Private/SSequencer.cpp
+++ b/Engine/Source/Editor/Sequencer/Private/SSequencer.cpp
@@ -72,7 +72,6 @@
 #include "ScopedTransaction.h"
 #include "SequencerTimeSliderController.h"
 #include "SequencerToolMenuContext.h"
-#include "SSequencerBakeTransform.h"
 #include "SSequencerSectionOverlay.h"
 #include "STemporarilyFocusedSpinBox.h"
 #include "MVVM/Views/STrackAreaView.h"
@@ -1186,11 +1185,6 @@ void SSequencer::BindCommands(TSharedRef<FUICommandList> SequencerCommandBinding
 		FExecuteAction::CreateLambda([this] { TransformBox->ToggleVisibility(); })
 	);
 
-	SequencerCommandBindings->MapAction(
-		FSequencerCommands::Get().BakeTransform,
-		FExecuteAction::CreateSP(this, &SSequencer::BakeTransform)
-	);
-
 	// Allow jumping to the Sequencer tree search if you have Sequencer focused
 	SequencerCommandBindings->MapAction(
 		FSequencerCommands::Get().QuickTreeSearch,
@@ -1755,11 +1749,6 @@ TSharedRef<SWidget> SSequencer::MakeActionsMenu()
 	// transform actions
 	MenuBuilder.BeginSection("Transform", LOCTEXT("TransformHeader", "Transform"));
 	
-	if (SequencerPtr.Pin()->IsLevelEditorSequencer())
-	{
-		MenuBuilder.AddMenuEntry(FSequencerCommands::Get().BakeTransform);
-	}
-	
 	MenuBuilder.AddMenuEntry(FSequencerCommands::Get().ToggleShowTransformBox);
 	MenuBuilder.AddMenuEntry(FSequencerCommands::Get().ToggleShowStretchBox);
 	MenuBuilder.AddMenuEntry(FSequencerCommands::Get().TranslateLeft);
@@ -3777,51 +3766,6 @@ void SSequencer::OnSetSequenceReadOnly(ECheckBoxState CheckBoxState)
 	}
 }
 
-void SSequencer::BakeTransform()
-{
-	TSharedPtr<FSequencer> Sequencer = SequencerPtr.Pin();
-	if (!Sequencer)
-	{
-		return;
-	}
-
-	UMovieSceneSequence* Sequence = Sequencer->GetFocusedMovieSceneSequence();
-	UMovieScene* MovieScene = Sequence ? Sequence->GetMovieScene() : nullptr;
-	if (!MovieScene)
-	{
-		return;
-	}
-
-	TArray<FGuid> ObjectBindings;
-	Sequencer->GetSelectedObjects(ObjectBindings);
-
-	TArray<FMovieSceneBindingProxy> BindingProxies;
-	for (FGuid Guid : ObjectBindings)
-	{
-		BindingProxies.Add(FMovieSceneBindingProxy(Guid, Sequence));
-	}
-
-	static FBakingAnimationKeySettings Settings; //reuse the settings except for the range
-	Settings.StartFrame = UE::MovieScene::DiscreteInclusiveLower(MovieScene->GetPlaybackRange());
-	Settings.EndFrame = UE::MovieScene::DiscreteExclusiveUpper(MovieScene->GetPlaybackRange());
-
-	TSharedRef<SBakeTransformWidget> BakeWidget =
-		SNew(SBakeTransformWidget)
-		.Settings(Settings)
-		.Sequencer(Sequencer.Get())
-		.OnBake_Lambda([this, &BindingProxies, Sequencer](FBakingAnimationKeySettings InSettings)
-		{
-			if (FSequencerUtilities::BakeTransform(Sequencer.ToSharedRef(), BindingProxies, InSettings))
-			{
-				Settings = InSettings;
-			}
-			
-			return FReply::Handled();
-		});
-
-	BakeWidget->OpenDialog(true);
-}
-
 void SSequencer::SetPlayTimeClampedByWorkingRange(double Frame)
 {
 	if (SequencerPtr.IsValid())
diff --git a/Engine/Source/Editor/Sequencer/Private/SSequencer.h b/Engine/Source/Editor/Sequencer/Private/SSequencer.h
--- a/Engine/Source/Editor/Sequencer/Private/SSequencer.h
+++ b/Engine/Source/Editor/Sequencer/Private/SSequencer.h
@@ -583,8 +583,6 @@ private:
 	/** Returns whether or not the Curve Editor is enabled. Allows us to bind to the Slate Enabled attribute. */
 	bool GetIsCurveEditorEnabled() const { return !GetIsSequenceReadOnly(); }
 
-	void BakeTransform();
-
 public:
 	/** On Paste Command */
 	void OnPaste();
diff --git a/Engine/Source/Editor/Sequencer/Private/SSequencerBakeTransform.cpp b/Engine/Source/Editor/Sequencer/Private/SSequencerBakeTransform.cpp
deleted file mode 100644
diff --git a/Engine/Source/Editor/Sequencer/Private/SSequencerBakeTransform.h b/Engine/Source/Editor/Sequencer/Private/SSequencerBakeTransform.h
deleted file mode 100644
diff --git a/Engine/Source/Editor/Sequencer/Private/SequencerCommands.cpp b/Engine/Source/Editor/Sequencer/Private/SequencerCommands.cpp
--- a/Engine/Source/Editor/Sequencer/Private/SequencerCommands.cpp
+++ b/Engine/Source/Editor/Sequencer/Private/SequencerCommands.cpp
@@ -204,8 +204,6 @@ void FSequencerCommands::RegisterCommands()
 
 	UI_COMMAND(ToggleLimitViewportSelection, "Limit Viewport Selection", "Toggles viewport selectability between Sequencer only objects and all objects.", EUserInterfaceActionType::ToggleButton, FInputChord());
 
-	UI_COMMAND(BakeTransform, "Bake Transform", "Bake transform in world space, removing any existing transform and attach tracks", EUserInterfaceActionType::Button, FInputChord());
-
 	UI_COMMAND(ToggleSidebarVisible, "Sidebar Details", "Toggles the visibility of the Sequencer sidebar.", EUserInterfaceActionType::ToggleButton, FInputChord());
 	UI_COMMAND(ToggleSelectionDetailsPoppedOut, "Popout Selection Details", "Toggles the state of the Sequencer selection details between docked in the Sequencer tab and being popped out to its own tab window", EUserInterfaceActionType::ToggleButton, FInputChord());
 
diff --git a/Engine/Source/Editor/Sequencer/Public/SequencerCommands.h b/Engine/Source/Editor/Sequencer/Public/SequencerCommands.h
--- a/Engine/Source/Editor/Sequencer/Public/SequencerCommands.h
+++ b/Engine/Source/Editor/Sequencer/Public/SequencerCommands.h
@@ -330,9 +330,6 @@ public:
 	/** Toggle the visibility of the stretch box. */
 	TSharedPtr< FUICommandInfo > ToggleShowStretchBox;
 
-	/** Bake transform tracks */
-	TSharedPtr< FUICommandInfo > BakeTransform;
-
 	/** Opens the director blueprint for a sequence. */
 	TSharedPtr< FUICommandInfo > OpenDirectorBlueprint;
 
