PR #14434: Removed ALevelSequenceActor dependency from FLevelSequencePlaybackContext

#Synced-CL 50990070

diff --git a/Engine/Plugins/MovieScene/LevelSequenceEditor/Source/LevelSequenceEditor/Private/Misc/LevelSequencePlaybackContext.cpp b/Engine/Plugins/MovieScene/LevelSequenceEditor/Source/LevelSequenceEditor/Private/Misc/LevelSequencePlaybackContext.cpp
--- a/Engine/Plugins/MovieScene/LevelSequenceEditor/Source/LevelSequenceEditor/Private/Misc/LevelSequencePlaybackContext.cpp
+++ b/Engine/Plugins/MovieScene/LevelSequenceEditor/Source/LevelSequenceEditor/Private/Misc/LevelSequencePlaybackContext.cpp
@@ -27,9 +27,9 @@ class SLevelSequenceContextPicker : public SCompoundWidget
 {
 public:
 
-	using FContextAndClient = TTuple<UWorld*, ALevelSequenceActor*>;
+	using FContextAndClient = TTuple<UWorld*, IMovieScenePlaybackClient*>;
 
-	DECLARE_DELEGATE_TwoParams(FOnSetPlaybackContextAndClient, UWorld*, ALevelSequenceActor*);
+	DECLARE_DELEGATE_TwoParams(FOnSetPlaybackContextAndClient, UWorld*, IMovieScenePlaybackClient*);
 
 	SLATE_BEGIN_ARGS(SLevelSequenceContextPicker){}
 
@@ -50,7 +50,7 @@ private:
 
 	TSharedRef<SWidget> BuildWorldPickerMenu();
 
-	static FText GetContextAndClientDescription(const UWorld* Context, const ALevelSequenceActor* Client);
+	static FText GetContextAndClientDescription(const UWorld* Context, const IMovieScenePlaybackClient* Client);
 	static FText GetWorldDescription(const UWorld* World);
 
 	FText GetCurrentContextAndClientText() const
@@ -103,16 +103,16 @@ private:
 		return GetDefault<ULevelSequenceEditorSettings>()->bAutoBindToSimulate;
 	}
 
-	void OnSetPlaybackContextAndClient(TWeakObjectPtr<UWorld> InContext, TWeakObjectPtr<ALevelSequenceActor> InClient)
+	void OnSetPlaybackContextAndClient(TWeakObjectPtr<UWorld> InContext, TWeakInterfacePtr<IMovieScenePlaybackClient> InClient)
 	{
 		if (UWorld* NewContext = InContext.Get())
 		{
-			ALevelSequenceActor* NewClient = InClient.Get();
+			IMovieScenePlaybackClient* NewClient = InClient.Get();
 			OnSetPlaybackContextAndClientEvent.ExecuteIfBound(NewContext, NewClient);
 		}
 	}
 
-	bool IsCurrentPlaybackContextAndClient(TWeakObjectPtr<UWorld> InContext, TWeakObjectPtr<ALevelSequenceActor> InClient)
+	bool IsCurrentPlaybackContextAndClient(TWeakObjectPtr<UWorld> InContext, TWeakInterfacePtr<IMovieScenePlaybackClient> InClient)
 	{
 		FContextAndClient ContextAndClient = PlaybackContextAndClientAttribute.Get();
 		return (InContext == ContextAndClient.Key && InClient == ContextAndClient.Value);
@@ -130,10 +130,15 @@ namespace MovieScene
 {
 
 /**
- * Finds all level sequence actors in the given world, and return those that point to the given sequence.
+ * Finds all actors that implement IMovieScenePlaybackClient in the given world, and return those that can preview the given sequence.
  */
-static void FindLevelSequenceActors(const UWorld* InWorld, const ULevelSequence* InLevelSequence, TArray<ALevelSequenceActor*>& OutActors)
+static void FindPlaybackClientActors(const UWorld* InWorld, const ULevelSequence* InLevelSequence, TArray<IMovieScenePlaybackClient*>& OutClients)
 {
+	if (!InWorld)
+	{
+		return;
+	}
+
 	ULevelInstanceSubsystem* LevelInstanceSubsystem = InWorld->GetSubsystem<ULevelInstanceSubsystem>();
 
 	for (const ULevel* Level : InWorld->GetLevels())
@@ -146,10 +151,10 @@ static void FindLevelSequenceActors(const UWorld* InWorld, const ULevelSequence*
 
 		for (AActor* Actor : Level->Actors)
 		{
-			ALevelSequenceActor* LevelSequenceActor = Cast<ALevelSequenceActor>(Actor);
-			if (LevelSequenceActor && LevelSequenceActor->GetSequence() == InLevelSequence)
+			IMovieScenePlaybackClient* PlaybackClient = Cast<IMovieScenePlaybackClient>(Actor);
+			if (InLevelSequence && PlaybackClient && PlaybackClient->CanPreview(*InLevelSequence))
 			{
-				OutActors.Add(LevelSequenceActor);
+				OutClients.Add(PlaybackClient);
 			}
 		}
 	}
@@ -215,7 +220,7 @@ UObject* FLevelSequencePlaybackContext::GetPlaybackContext() const
 {
 	UpdateCachedContextAndClient();
 
-	if (ALevelSequenceActor* Client = GetPlaybackClient())
+	if (UObject* Client = GetPlaybackClientAsUObject())
 	{
 		return Client;
 	}
@@ -228,15 +233,21 @@ UObject* FLevelSequencePlaybackContext::GetPlaybackContextAsObject() const
 	return GetPlaybackContext();
 }
 
-ALevelSequenceActor* FLevelSequencePlaybackContext::GetPlaybackClient() const
+UObject* FLevelSequencePlaybackContext::GetPlaybackClientAsUObject() const
 {
 	UpdateCachedContextAndClient();
-	return WeakCurrentClient.Get();
+	return WeakCurrentClient.GetObject();
+}
+
+ALevelSequenceActor* FLevelSequencePlaybackContext::GetPlaybackClient() const
+{
+	return Cast<ALevelSequenceActor>(GetPlaybackClientAsUObject());
 }
 
 IMovieScenePlaybackClient* FLevelSequencePlaybackContext::GetPlaybackClientAsInterface() const
 {
-	return GetPlaybackClient();
+	UpdateCachedContextAndClient();
+	return WeakCurrentClient.Get();
 }
 
 TArray<UObject*> FLevelSequencePlaybackContext::GetEventContexts() const
@@ -252,7 +263,7 @@ TArray<UObject*> FLevelSequencePlaybackContext::GetEventContexts() const
 	return Contexts;
 }
 
-void FLevelSequencePlaybackContext::OverrideWith(UWorld* InNewContext, ALevelSequenceActor* InNewClient)
+void FLevelSequencePlaybackContext::OverrideWith(UWorld* InNewContext, IMovieScenePlaybackClient* InNewClient)
 {
 	// InNewContext may be null to force an auto update
 	WeakCurrentContext = InNewContext;
@@ -298,9 +309,9 @@ FLevelSequencePlaybackContext::FContextAndClient FLevelSequencePlaybackContext::
 			const bool bIsServerWorld = (ThisWorld && ThisWorld->GetNetDriver() && ThisWorld->GetNetDriver()->IsServer());
 			if (bIsPIEValid && bAllowPlaybackContextBinding && RecordingWorld != ThisWorld && !bIsServerWorld)
 			{
-				TArray<ALevelSequenceActor*> LevelSequenceActors;
-				UE::MovieScene::FindLevelSequenceActors(ThisWorld, InLevelSequence, LevelSequenceActors);
-				return FContextAndClient(ThisWorld, (LevelSequenceActors.Num() > 0 ? LevelSequenceActors[0] : nullptr));
+				TArray<IMovieScenePlaybackClient*> PlaybackClients;
+				UE::MovieScene::FindPlaybackClientActors(ThisWorld, InLevelSequence, PlaybackClients);
+				return FContextAndClient(ThisWorld, (PlaybackClients.Num() > 0 ? PlaybackClients[0] : nullptr));
 			}
 		}
 		else if (Context.WorldType == EWorldType::Editor)
@@ -314,16 +325,16 @@ FLevelSequencePlaybackContext::FContextAndClient FLevelSequencePlaybackContext::
 		// First search selected actors. If one is selected, use that as the client.
 		for (FSelectionIterator It(GEditor->GetSelectedActorIterator()); It; ++It)
 		{
-			if (ALevelSequenceActor* LSA = Cast<ALevelSequenceActor>(*It))
+			if (IMovieScenePlaybackClient* Client = Cast<IMovieScenePlaybackClient>(*It))
 			{
-				return FContextAndClient(EditorWorld, LSA);
+				return FContextAndClient(EditorWorld, Client);
 			}
 		}
 
 		// Otherwise, attempt to find one in the world.
-		TArray<ALevelSequenceActor*> LevelSequenceActors;
-		UE::MovieScene::FindLevelSequenceActors(EditorWorld, InLevelSequence, LevelSequenceActors);
-		return FContextAndClient(EditorWorld, (LevelSequenceActors.Num() > 0 ? LevelSequenceActors[0] : nullptr));
+		TArray<IMovieScenePlaybackClient*> PlaybackClients;
+		UE::MovieScene::FindPlaybackClientActors(EditorWorld, InLevelSequence, PlaybackClients);
+		return FContextAndClient(EditorWorld, (PlaybackClients.Num() > 0 ? PlaybackClients[0] : nullptr));
 	}
 
 	return FContextAndClient(nullptr, nullptr);
@@ -381,15 +392,17 @@ void SLevelSequenceContextPicker::Construct(const FArguments& InArgs)
 	];
 }
 
-FText SLevelSequenceContextPicker::GetContextAndClientDescription(const UWorld* Context, const ALevelSequenceActor* Client)
+FText SLevelSequenceContextPicker::GetContextAndClientDescription(const UWorld* Context, const IMovieScenePlaybackClient* Client)
 {
 	const FText WorldDescription = GetWorldDescription(Context);
 	if (Client)
 	{
+		const AActor* ClientActor = Cast<AActor>(Client);
+		const FString ClientName = ClientActor ? ClientActor->GetActorNameOrLabel() : GetNameSafe(Cast<UObject>(Client));
 		return FText::Format(
 				LOCTEXT("PlaybackContextDescription", "{0} ({1})"), 
 				WorldDescription,
-				FText::FromString(Client->GetName()));
+				FText::FromString(ClientName));
 	}
 	else
 	{
@@ -445,29 +458,29 @@ TSharedRef<SWidget> SLevelSequenceContextPicker::BuildWorldPickerMenu()
 				continue;
 			}
 
-			bool bFoundActors = false;
+			bool bFoundClients = false;
 			if (LevelSequence)
 			{
-				TArray<ALevelSequenceActor*> LevelSequenceActors;
-				UE::MovieScene::FindLevelSequenceActors(World, LevelSequence, LevelSequenceActors);
-				bFoundActors = LevelSequenceActors.Num() > 0;
+				TArray<IMovieScenePlaybackClient*> PlaybackClients;
+				UE::MovieScene::FindPlaybackClientActors(World, LevelSequence, PlaybackClients);
+				bFoundClients = PlaybackClients.Num() > 0;
 
-				for (ALevelSequenceActor* LevelSequenceActor : LevelSequenceActors)
+				for (IMovieScenePlaybackClient* PlaybackClient : PlaybackClients)
 				{
 					MenuBuilder.AddMenuEntry(
-							GetContextAndClientDescription(World, LevelSequenceActor),
+							GetContextAndClientDescription(World, PlaybackClient),
 							FText(),
 							FSlateIcon(),
 							FUIAction(
 								FExecuteAction::CreateSP(
 									this, 
 									&SLevelSequenceContextPicker::OnSetPlaybackContextAndClient, 
-									MakeWeakObjectPtr(World), MakeWeakObjectPtr(LevelSequenceActor)),
+									MakeWeakObjectPtr(World), TWeakInterfacePtr<IMovieScenePlaybackClient>(PlaybackClient)),
 								FCanExecuteAction(),
 								FIsActionChecked::CreateSP(
 									this, 
 									&SLevelSequenceContextPicker::IsCurrentPlaybackContextAndClient, 
-									MakeWeakObjectPtr(World), MakeWeakObjectPtr(LevelSequenceActor))
+									MakeWeakObjectPtr(World), TWeakInterfacePtr<IMovieScenePlaybackClient>(PlaybackClient))
 								),
 							NAME_None,
 							EUserInterfaceActionType::RadioButton
@@ -475,7 +488,7 @@ TSharedRef<SWidget> SLevelSequenceContextPicker::BuildWorldPickerMenu()
 				}
 			}
 			
-			if (!bFoundActors)
+			if (!bFoundClients)
 			{
 				MenuBuilder.AddMenuEntry(
 						GetContextAndClientDescription(World, nullptr),
@@ -485,12 +498,12 @@ TSharedRef<SWidget> SLevelSequenceContextPicker::BuildWorldPickerMenu()
 							FExecuteAction::CreateSP(
 								this, 
 								&SLevelSequenceContextPicker::OnSetPlaybackContextAndClient, 
-								MakeWeakObjectPtr(World), MakeWeakObjectPtr<ALevelSequenceActor>(nullptr)),
+								MakeWeakObjectPtr(World), TWeakInterfacePtr<IMovieScenePlaybackClient>()),
 							FCanExecuteAction(),
 							FIsActionChecked::CreateSP(
 								this, 
 								&SLevelSequenceContextPicker::IsCurrentPlaybackContextAndClient, 
-								MakeWeakObjectPtr(World), MakeWeakObjectPtr<ALevelSequenceActor>(nullptr))
+								MakeWeakObjectPtr(World), TWeakInterfacePtr<IMovieScenePlaybackClient>())
 							),
 						NAME_None,
 						EUserInterfaceActionType::RadioButton
diff --git a/Engine/Plugins/MovieScene/LevelSequenceEditor/Source/LevelSequenceEditor/Public/Misc/LevelSequencePlaybackContext.h b/Engine/Plugins/MovieScene/LevelSequenceEditor/Source/LevelSequenceEditor/Public/Misc/LevelSequencePlaybackContext.h
--- a/Engine/Plugins/MovieScene/LevelSequenceEditor/Source/LevelSequenceEditor/Public/Misc/LevelSequencePlaybackContext.h
+++ b/Engine/Plugins/MovieScene/LevelSequenceEditor/Source/LevelSequenceEditor/Public/Misc/LevelSequencePlaybackContext.h
@@ -3,6 +3,7 @@
 #pragma once
 
 #include "UObject/WeakObjectPtr.h"
+#include "UObject/WeakInterfacePtr.h"
 
 class SWidget;
 
@@ -45,6 +46,9 @@ public:
 	/**
 	 * Resolve the current playback client. May be nullptr.
 	 */
+	UE_API UObject* GetPlaybackClientAsUObject() const;
+
+	UE_DEPRECATED(5.8, "GetPlaybackClient has been deprecated, please use GetPlaybackClientAsUObject")
 	UE_API ALevelSequenceActor* GetPlaybackClient() const;
 
 	/**
@@ -61,11 +65,11 @@ public:
 	 * Specify a new world to use as the context. Persists until the next PIE or map change event.
 	 * May be null, in which case the context will be recomputed automatically
 	 */
-	UE_API void OverrideWith(UWorld* InNewContext, ALevelSequenceActor* InNewClient);
+	UE_API void OverrideWith(UWorld* InNewContext, IMovieScenePlaybackClient* InNewClient);
 
 private:
 
-	using FContextAndClient = TTuple<UWorld*, ALevelSequenceActor*>;
+	using FContextAndClient = TTuple<UWorld*, IMovieScenePlaybackClient*>;
 
 	/**
 	 * Compute the new playback context based on the user's current auto-bind settings.
@@ -81,7 +85,7 @@ private:
 	/**
 	 * Gets both the context and client.
 	 */
-	//TTuple<UWorld*, ALevelSequenceActor*>
+	//TTuple<UWorld*, IMovieScenePlaybackClient*>
 	UE_API FContextAndClient GetPlaybackContextAndClient() const;
 
 	UE_API void OnPieEvent(bool);
@@ -97,7 +101,7 @@ private:
 	mutable TWeakObjectPtr<UWorld> WeakCurrentContext;
 
 	/** Mutable cached client pointer */
-	mutable TWeakObjectPtr<ALevelSequenceActor> WeakCurrentClient;
+	mutable TWeakInterfacePtr<IMovieScenePlaybackClient> WeakCurrentClient;
 };
 
 #undef UE_API
diff --git a/Engine/Source/Runtime/LevelSequence/Private/LevelSequenceActor.cpp b/Engine/Source/Runtime/LevelSequence/Private/LevelSequenceActor.cpp
--- a/Engine/Source/Runtime/LevelSequence/Private/LevelSequenceActor.cpp
+++ b/Engine/Source/Runtime/LevelSequence/Private/LevelSequenceActor.cpp
@@ -146,6 +146,13 @@ bool ALevelSequenceActor::GetIsReplicatedPlayback() const
 	return bReplicatePlayback;
 }
 
+#if WITH_EDITOR
+bool ALevelSequenceActor::CanPreview(const UMovieSceneSequence& InSequence) const
+{
+	return GetSequence() == &InSequence;
+}
+#endif
+
 ULevelSequencePlayer* ALevelSequenceActor::GetSequencePlayer() const
 {
 PRAGMA_DISABLE_DEPRECATION_WARNINGS // make SequencePlayer protected and remove this for 5.6
diff --git a/Engine/Source/Runtime/LevelSequence/Public/LevelSequenceActor.h b/Engine/Source/Runtime/LevelSequence/Public/LevelSequenceActor.h
--- a/Engine/Source/Runtime/LevelSequence/Public/LevelSequenceActor.h
+++ b/Engine/Source/Runtime/LevelSequence/Public/LevelSequenceActor.h
@@ -259,6 +259,9 @@ protected:
 	LEVELSEQUENCE_API virtual UObject* GetInstanceData() const override;
 	LEVELSEQUENCE_API virtual TOptional<EAspectRatioAxisConstraint> GetAspectRatioAxisConstraint() const override;
 	LEVELSEQUENCE_API virtual bool GetIsReplicatedPlayback() const override;
+#if WITH_EDITOR
+	LEVELSEQUENCE_API virtual bool CanPreview(const UMovieSceneSequence& InSequence) const override;
+#endif
 	//~ End IMovieScenePlaybackClient interface
 
 	//~ Begin UObject interface
diff --git a/Engine/Source/Runtime/MovieScene/Public/IMovieScenePlaybackClient.h b/Engine/Source/Runtime/MovieScene/Public/IMovieScenePlaybackClient.h
--- a/Engine/Source/Runtime/MovieScene/Public/IMovieScenePlaybackClient.h
+++ b/Engine/Source/Runtime/MovieScene/Public/IMovieScenePlaybackClient.h
@@ -10,6 +10,7 @@
 #include "IMovieScenePlaybackClient.generated.h"
 
 struct FMovieSceneEvaluationRange;
+class UMovieSceneSequence;
 
 /** Movie scene binding overrides interface */
 UINTERFACE(MinimalAPI)
@@ -63,4 +64,15 @@ public:
 	 * This gives clients an opportunity to apply time-warping effects without manipulating the actual user-facing playback position
 	 */
 	virtual void WarpEvaluationRange(FMovieSceneEvaluationRange& InOutRange) const {}
+	
+#if WITH_EDITOR
+	/*
+	 * Whether this playback client can play the given sequence in editor preview.
+	 * This gives client actors an opportunity to add themselves to the list of preview contexts in the level editor.
+	 */
+	virtual bool CanPreview(const UMovieSceneSequence& InSequence) const
+	{ 
+		return false; 
+	}
+#endif
 };
