Sequencer: Follow up to 51220238 to fix subsequence range when limited by shot

#Synced-CL 51926426

diff --git a/Engine/Plugins/MovieScene/LevelSequenceEditor/Source/LevelSequenceEditor/Private/CinematicViewport/SCinematicTransportRange.cpp b/Engine/Plugins/MovieScene/LevelSequenceEditor/Source/LevelSequenceEditor/Private/CinematicViewport/SCinematicTransportRange.cpp
--- a/Engine/Plugins/MovieScene/LevelSequenceEditor/Source/LevelSequenceEditor/Private/CinematicViewport/SCinematicTransportRange.cpp
+++ b/Engine/Plugins/MovieScene/LevelSequenceEditor/Source/LevelSequenceEditor/Private/CinematicViewport/SCinematicTransportRange.cpp
@@ -89,11 +89,10 @@ void SCinematicTransportRange::SetTime(const FGeometry& MyGeometry, const FPoint
 		FFrameTime ScrubTime = NewTimeSeconds * TickResolution;
 
 		// Clamp first, snap to frame last
-		TRange<FFrameNumber> PlayheadScrubbingRange = Sequencer->GetSequencerSettings()->GetPlayheadScrubbingRange(Sequencer->GetRootMovieSceneSequence()->GetMovieScene());
+		TRange<FFrameNumber> PlayheadScrubbingRange = Sequencer->GetSequencerSettings()->GetPlayheadScrubbingRange(WeakSequencer);
 		if (!PlayheadScrubbingRange.IsEmpty())
 		{
-			TRange<FFrameNumber> PlaybackRange = Sequencer->GetSubSequenceRange().Get(PlayheadScrubbingRange);
-			ScrubTime = UE::MovieScene::ClampToDiscreteRange(ScrubTime, PlaybackRange);
+			ScrubTime = UE::MovieScene::ClampToDiscreteRange(ScrubTime, PlayheadScrubbingRange);
 		}
 
 		ENearestKeyOption NearestKeyOption = ENearestKeyOption::NKO_None;
diff --git a/Engine/Source/Editor/Sequencer/Private/Sequencer.cpp b/Engine/Source/Editor/Sequencer/Private/Sequencer.cpp
--- a/Engine/Source/Editor/Sequencer/Private/Sequencer.cpp
+++ b/Engine/Source/Editor/Sequencer/Private/Sequencer.cpp
@@ -3764,11 +3764,10 @@ void FSequencer::PlayTo(FMovieSceneSequencePlaybackParams PlaybackParams)
 void FSequencer::SnapSequencerTime(FFrameTime& ScrubTime)
 {
 	// Clamp first, snap to frame last
-	TRange<FFrameNumber> PlayheadScrubbingRange = GetSequencerSettings()->GetPlayheadScrubbingRange(GetRootMovieSceneSequence()->GetMovieScene());
+	TRange<FFrameNumber> PlayheadScrubbingRange = GetSequencerSettings()->GetPlayheadScrubbingRange(TWeakPtr<ISequencer>(AsShared()));
 	if (!PlayheadScrubbingRange.IsEmpty())
 	{
-		TRange<FFrameNumber> PlaybackRange = GetSubSequenceRange().Get(PlayheadScrubbingRange);
-		ScrubTime = UE::MovieScene::ClampToDiscreteRange(ScrubTime, PlaybackRange);
+		ScrubTime = UE::MovieScene::ClampToDiscreteRange(ScrubTime, PlayheadScrubbingRange);
 	}
 
 	if (GetSequencerSettings()->GetForceWholeFrames())
diff --git a/Engine/Source/Editor/Sequencer/Private/SequencerSettings.cpp b/Engine/Source/Editor/Sequencer/Private/SequencerSettings.cpp
--- a/Engine/Source/Editor/Sequencer/Private/SequencerSettings.cpp
+++ b/Engine/Source/Editor/Sequencer/Private/SequencerSettings.cpp
@@ -526,15 +526,22 @@ void USequencerSettings::SetPlayheadScrubbingLimit(ESequencerPlayheadScrubbingLi
 	}
 }
 
-TRange<FFrameNumber> USequencerSettings::GetPlayheadScrubbingRange(UMovieScene* MovieScene) const
+TRange<FFrameNumber> USequencerSettings::GetPlayheadScrubbingRange(TWeakPtr<ISequencer> InSequencer) const
 {
+	TSharedPtr<ISequencer> Sequencer = InSequencer.Pin();
+	if (!Sequencer)
+	{
+		return TRange<FFrameNumber>::Empty();
+	}
+
+	TRange<FFrameNumber> PlayheadScrubbingRange;
 	if (PlayheadScrubbingLimit == ESequencerPlayheadScrubbingLimit::SPSL_PlaybackRange)
 	{
-		return MovieScene->GetPlaybackRange();
+		return PlayheadScrubbingRange = Sequencer->GetSubSequenceRange().Get(Sequencer->GetFocusedMovieSceneSequence()->GetMovieScene()->GetPlaybackRange());
 	}
 	else if (PlayheadScrubbingLimit == ESequencerPlayheadScrubbingLimit::SPSL_SelectionRange)
 	{
-		return MovieScene->GetSelectionRange();
+		return PlayheadScrubbingRange = Sequencer->GetFocusedMovieSceneSequence()->GetMovieScene()->GetSelectionRange();
 	}
 	return TRange<FFrameNumber>::Empty();
 }
diff --git a/Engine/Source/Editor/Sequencer/Public/SequencerSettings.h b/Engine/Source/Editor/Sequencer/Public/SequencerSettings.h
--- a/Engine/Source/Editor/Sequencer/Public/SequencerSettings.h
+++ b/Engine/Source/Editor/Sequencer/Public/SequencerSettings.h
@@ -16,6 +16,7 @@
 
 #define UE_API SEQUENCER_API
 
+class ISequencer;
 struct FPropertyChangedEvent;
 enum class EAutoChangeMode : uint8;
 enum class EAllowEditsMode : uint8;
@@ -337,7 +338,7 @@ public:
 	/** Set the limit that the playhead should be restricted to when scrubbing */
 	UE_API void SetPlayheadScrubbingLimit(ESequencerPlayheadScrubbingLimit InPlayheadScrubbingLimit);
 	/** Get the actual frame number range that the playhead should be restricted to when scrubbing */
-	UE_API TRange<FFrameNumber> GetPlayheadScrubbingRange(UMovieScene* MovieScene) const;
+	UE_API TRange<FFrameNumber> GetPlayheadScrubbingRange(TWeakPtr<ISequencer> InSequencer) const;
 
 	/** @return true if the playback range should be synced to the section bounds, false otherwise */
 	UE_API bool ShouldKeepPlayRangeInSectionBounds() const;
