Sequencer: Expose ISequencer Playback delegates to the level editor subsystem (slack request).

#Synced-CL 52367967

diff --git a/Engine/Plugins/MovieScene/LevelSequenceEditor/Source/LevelSequenceEditor/Private/LevelSequenceEditorSubsystem.cpp b/Engine/Plugins/MovieScene/LevelSequenceEditor/Source/LevelSequenceEditor/Private/LevelSequenceEditorSubsystem.cpp
--- a/Engine/Plugins/MovieScene/LevelSequenceEditor/Source/LevelSequenceEditor/Private/LevelSequenceEditorSubsystem.cpp
+++ b/Engine/Plugins/MovieScene/LevelSequenceEditor/Source/LevelSequenceEditor/Private/LevelSequenceEditorSubsystem.cpp
@@ -685,6 +685,8 @@ void ULevelSequenceEditorSubsystem::OnSequencerCreated(TSharedRef<ISequencer> In
 	Sequencers.Add(TWeakPtr<ISequencer>(InSequencer));
 	InSequencer->OnCloseEvent().AddUObject(this, &ULevelSequenceEditorSubsystem::OnSequencerClosed);
 
+	CreatePlaybackDelegates(InSequencer);
+
 	// Create a new binding list that will be used with the sequencer that was just added
 	UMovieSceneBindingPropertyInfoList* BindingPropertyInfoList = GetBindingPropertyInfoList(InSequencer);
 	if (!BindingPropertyInfoList)
@@ -703,9 +705,60 @@ void ULevelSequenceEditorSubsystem::OnSequencerClosed(TSharedRef<ISequencer> InS
 		BindingPropertyInfoLists.RemoveAt(IndexToRemove);
 	}
 
+	if (FSequencerPlaybackDelegateHandles* Handles = SequencerPlaybackDelegateHandles.Find(InSequencer))
+	{
+		InSequencer->OnGlobalTimeChanged().Remove(Handles->GlobalTimeChanged);
+		InSequencer->OnPlayEvent().Remove(Handles->PlayEvent);
+		InSequencer->OnStopEvent().Remove(Handles->StopEvent);
+		InSequencer->OnRecordEvent().Remove(Handles->RecordEvent);
+		InSequencer->OnBeginScrubbingEvent().Remove(Handles->BeginScrubbing);
+		InSequencer->OnEndScrubbingEvent().Remove(Handles->EndScrubbing);
+		SequencerPlaybackDelegateHandles.Remove(InSequencer);
+	}
+
 	TrackRowMetadataHelperList.Empty();
 }
 
+void ULevelSequenceEditorSubsystem::CreatePlaybackDelegates(TSharedRef<ISequencer> InSequencer)
+{
+	FSequencerPlaybackDelegateHandles Handles;
+	Handles.GlobalTimeChanged = InSequencer->OnGlobalTimeChanged().AddUObject(this, &ULevelSequenceEditorSubsystem::HandleGlobalTimeChanged);
+	Handles.PlayEvent         = InSequencer->OnPlayEvent().AddUObject(this, &ULevelSequenceEditorSubsystem::HandlePlayEvent);
+	Handles.StopEvent         = InSequencer->OnStopEvent().AddUObject(this, &ULevelSequenceEditorSubsystem::HandleStopEvent);
+	Handles.RecordEvent       = InSequencer->OnRecordEvent().AddUObject(this, &ULevelSequenceEditorSubsystem::HandleRecordEvent);
+	Handles.BeginScrubbing    = InSequencer->OnBeginScrubbingEvent().AddUObject(this, &ULevelSequenceEditorSubsystem::HandleBeginScrubbingEvent);
+	Handles.EndScrubbing      = InSequencer->OnEndScrubbingEvent().AddUObject(this, &ULevelSequenceEditorSubsystem::HandleEndScrubbingEvent);
+	SequencerPlaybackDelegateHandles.Add(InSequencer, Handles);
+}
+
+void ULevelSequenceEditorSubsystem::HandleGlobalTimeChanged() 
+{
+	OnGlobalTimeChanged.Broadcast(); 
+}
+
+void ULevelSequenceEditorSubsystem::HandlePlayEvent() 
+{
+	OnPlayEvent.Broadcast(); 
+}
+
+void ULevelSequenceEditorSubsystem::HandleStopEvent()  
+{
+	OnStopEvent.Broadcast(); 
+}
+void ULevelSequenceEditorSubsystem::HandleRecordEvent()  
+{ 
+	OnRecordEvent.Broadcast(); 
+}
+void ULevelSequenceEditorSubsystem::HandleBeginScrubbingEvent() 
+{
+	OnBeginScrubbingEvent.Broadcast(); 
+}
+
+void ULevelSequenceEditorSubsystem::HandleEndScrubbingEvent()   
+{ 
+	OnEndScrubbingEvent.Broadcast(); 
+}
+
 void ULevelSequenceEditorSubsystem::AddBindingDetailCustomizations(TSharedRef<IDetailsView> DetailsView, TSharedPtr<ISequencer> ActiveSequencer, FGuid BindingGuid)
 {
 	// TODO: Do we want to create a generalized way for folks to add instanced property layouts for other custom binding types so they can have access to sequencer context?
diff --git a/Engine/Plugins/MovieScene/LevelSequenceEditor/Source/LevelSequenceEditor/Public/LevelSequenceEditorSubsystem.h b/Engine/Plugins/MovieScene/LevelSequenceEditor/Source/LevelSequenceEditor/Public/LevelSequenceEditorSubsystem.h
--- a/Engine/Plugins/MovieScene/LevelSequenceEditor/Source/LevelSequenceEditor/Public/LevelSequenceEditorSubsystem.h
+++ b/Engine/Plugins/MovieScene/LevelSequenceEditor/Source/LevelSequenceEditor/Public/LevelSequenceEditorSubsystem.h
@@ -30,6 +30,13 @@ struct FMovieScenePossessable;
 
 DECLARE_LOG_CATEGORY_EXTERN(LogLevelSequenceEditor, Log, All);
 
+DECLARE_DYNAMIC_MULTICAST_DELEGATE(FOnLevelSequenceGlobalTimeChanged);
+DECLARE_DYNAMIC_MULTICAST_DELEGATE(FOnLevelSequencePlayEvent);
+DECLARE_DYNAMIC_MULTICAST_DELEGATE(FOnLevelSequenceStopEvent);
+DECLARE_DYNAMIC_MULTICAST_DELEGATE(FOnLevelSequenceRecordEvent);
+DECLARE_DYNAMIC_MULTICAST_DELEGATE(FOnLevelSequenceBeginScrubbingEvent);
+DECLARE_DYNAMIC_MULTICAST_DELEGATE(FOnLevelSequenceEndScrubbingEvent);
+
 class ACineCameraActor;
 class FExtender;
 class FMenuBuilder;
@@ -114,6 +121,30 @@ public:
 	UFUNCTION(BlueprintPure, Category = "Level Sequence Editor")
 	UE_API USequencerCurveEditorObject* GetCurveEditor();
 
+	/** Called when the global time changes in the active sequencer */
+	UPROPERTY(BlueprintAssignable, Category = "Level Sequence Editor")
+	FOnLevelSequenceGlobalTimeChanged OnGlobalTimeChanged;
+
+	/** Called when the active sequencer begins playing */
+	UPROPERTY(BlueprintAssignable, Category = "Level Sequence Editor")
+	FOnLevelSequencePlayEvent OnPlayEvent;
+
+	/** Called when the active sequencer stops playing */
+	UPROPERTY(BlueprintAssignable, Category = "Level Sequence Editor")
+	FOnLevelSequenceStopEvent OnStopEvent;
+
+	/** Called when the active sequencer toggles recording */
+	UPROPERTY(BlueprintAssignable, Category = "Level Sequence Editor")
+	FOnLevelSequenceRecordEvent OnRecordEvent;
+
+	/** Called when the user begins scrubbing in the active sequencer */
+	UPROPERTY(BlueprintAssignable, Category = "Level Sequence Editor")
+	FOnLevelSequenceBeginScrubbingEvent OnBeginScrubbingEvent;
+
+	/** Called when the user ends scrubbing in the active sequencer */
+	UPROPERTY(BlueprintAssignable, Category = "Level Sequence Editor")
+	FOnLevelSequenceEndScrubbingEvent OnEndScrubbingEvent;
+
 	/** Add existing actors to Sequencer. Tracks will be automatically added based on default track settings. */
 	UFUNCTION(BlueprintCallable, Category = "Level Sequence Editor")
 	UE_API TArray<FMovieSceneBindingProxy> AddActors(const TArray<AActor*>& Actors);
@@ -345,6 +376,26 @@ private:
 
 	UMovieSceneBindingPropertyInfoList* GetBindingPropertyInfoList(TSharedRef<ISequencer> Sequencer);
 
+	UE_API void CreatePlaybackDelegates(TSharedRef<ISequencer> InSequencer);
+
+	void HandleGlobalTimeChanged();
+	void HandlePlayEvent();
+	void HandleStopEvent();
+	void HandleRecordEvent();
+	void HandleBeginScrubbingEvent();
+	void HandleEndScrubbingEvent();
+
+	struct FSequencerPlaybackDelegateHandles
+	{
+		FDelegateHandle GlobalTimeChanged;
+		FDelegateHandle PlayEvent;
+		FDelegateHandle StopEvent;
+		FDelegateHandle RecordEvent;
+		FDelegateHandle BeginScrubbing;
+		FDelegateHandle EndScrubbing;
+	};
+	TMap<TWeakPtr<ISequencer>, FSequencerPlaybackDelegateHandles> SequencerPlaybackDelegateHandles;
+
 	FDelegateHandle OnSequencerCreatedHandle;
 
 	/* List of sequencers that have been created */
