Sequencer: Fix command binding for overlay options.

The command list was being populated each time the submenu opened which created lambdas pointing to the transient SFilmOverlayOptions widget. That can get destroyed when the toolbar refreshes. The fix here

#Synced-CL 52445895

diff --git a/Engine/Plugins/MovieScene/LevelSequenceEditor/Source/LevelSequenceEditor/Private/CinematicViewport/FilmOverlays.cpp b/Engine/Plugins/MovieScene/LevelSequenceEditor/Source/LevelSequenceEditor/Private/CinematicViewport/FilmOverlays.cpp
--- a/Engine/Plugins/MovieScene/LevelSequenceEditor/Source/LevelSequenceEditor/Private/CinematicViewport/FilmOverlays.cpp
+++ b/Engine/Plugins/MovieScene/LevelSequenceEditor/Source/LevelSequenceEditor/Private/CinematicViewport/FilmOverlays.cpp
@@ -488,44 +488,64 @@ void SFilmOverlay::Construct(const FArguments& InArgs)
 
 int32 SFilmOverlay::OnPaint(const FPaintArgs& Args, const FGeometry& AllottedGeometry, const FSlateRect& MyCullingRect, FSlateWindowElementList& OutDrawElements, int32 LayerId, const FWidgetStyle& InWidgetStyle, bool bParentEnabled) const
 {
-	TArray<IFilmOverlay*> Overlays = FilmOverlays.Get();
-	for (IFilmOverlay* Overlay : Overlays)
+	if (!PrimaryFilmOverlay.IsNone())
 	{
-		Overlay->Paint(AllottedGeometry, MyCullingRect, OutDrawElements, LayerId);
-		++LayerId;
+		if (const TSharedPtr<IFilmOverlay>* Found = UFilmOverlayToolkit::GetPrimaryFilmOverlays().Find(PrimaryFilmOverlay))
+		{
+			if (Found->IsValid())
+			{
+				(*Found)->Paint(AllottedGeometry, MyCullingRect, OutDrawElements, LayerId);
+				++LayerId;
+			}
+		}
+	}
+
+	for (const TPair<FName, TSharedPtr<IFilmOverlay>>& Pair : UFilmOverlayToolkit::GetToggleableFilmOverlays())
+	{
+		if (Pair.Value.IsValid() && Pair.Value->IsEnabled())
+		{
+			Pair.Value->Paint(AllottedGeometry, MyCullingRect, OutDrawElements, LayerId);
+			++LayerId;
+		}
 	}
 
 	return LayerId;
 }
 
-void SFilmOverlayOptions::Construct(const FArguments& InArgs, TSharedPtr<SFilmOverlay> InFilmOverlay)
+// static
+void SFilmOverlayOptions::RegisterDefaultOverlays()
 {
-	PrimaryColorTint = FLinearColor(1.f, 1.f, 1.f, .5f);
-
 	auto& ExistingPrimaryOverlays = UFilmOverlayToolkit::GetPrimaryFilmOverlays();
-	
+
 #define PREPARE_PRIMARY_OVERLAY(Name, Construct) { TSharedPtr<IFilmOverlay> Overlay = ExistingPrimaryOverlays.FindRef(Name); \
 	if (!Overlay) { UFilmOverlayToolkit::RegisterPrimaryFilmOverlay(Name, Construct); }}
-	
+
 	PREPARE_PRIMARY_OVERLAY(NAME_None, MakeShared<FFilmOverlay_None>());
 	PREPARE_PRIMARY_OVERLAY("Grid3x3", MakeShared<FFilmOverlay_Grid>(3, 3));
 	PREPARE_PRIMARY_OVERLAY("Grid2x2", MakeShared<FFilmOverlay_Grid>(2, 2));
 	PREPARE_PRIMARY_OVERLAY("Crosshair", MakeShared<FFilmOverlay_Crosshair>());
 	PREPARE_PRIMARY_OVERLAY("Rabatment", MakeShared<FFilmOverlay_Rabatment>());
-	
+
 #undef PREPARE_PRIMARY_OVERLAY
 
-	auto& ExistingToggleableOverplays = UFilmOverlayToolkit::GetToggleableFilmOverlays();
+	auto& ExistingToggleableOverlays = UFilmOverlayToolkit::GetToggleableFilmOverlays();
 
-#define PREPARE_TOGGLEABLE_OVERLAY(Name, Construct) { TSharedPtr<IFilmOverlay> Overlay = ExistingToggleableOverplays.FindRef(Name); \
+#define PREPARE_TOGGLEABLE_OVERLAY(Name, Construct) { TSharedPtr<IFilmOverlay> Overlay = ExistingToggleableOverlays.FindRef(Name); \
 	if (!Overlay) { UFilmOverlayToolkit::RegisterToggleableFilmOverlay(Name, Construct); }}
-	
+
 	PREPARE_TOGGLEABLE_OVERLAY("ActionSafe", MakeShared<FFilmOverlay_SafeFrame>(LOCTEXT("ActionSafe", "Action Safe"), 95.f, FLinearColor::Red));
 	PREPARE_TOGGLEABLE_OVERLAY("TitleSafe", MakeShared<FFilmOverlay_SafeFrame>(LOCTEXT("TitleSafe", "Title Safe"), 90.f, FLinearColor::Yellow));
 	PREPARE_TOGGLEABLE_OVERLAY("CustomSafe", MakeShared<FFilmOverlay_SafeFrame>(LOCTEXT("CustomSafe", "Custom Safe"), 85.f, FLinearColor::Green));
 	PREPARE_TOGGLEABLE_OVERLAY("LetterBox", MakeShared<FFilmOverlay_LetterBox>());
-	
+
 #undef PREPARE_TOGGLEABLE_OVERLAY
+}
+
+void SFilmOverlayOptions::Construct(const FArguments& InArgs, TSharedPtr<SFilmOverlay> InFilmOverlay)
+{
+	PrimaryColorTint = FLinearColor(1.f, 1.f, 1.f, .5f);
+
+	RegisterDefaultOverlays();
 
 	// used to choose between creating the whole button + menu widget, or just the menu widget
 	// e.g. Old toolbar uses the combo button, while the new one does not
@@ -569,8 +589,6 @@ void SFilmOverlayOptions::Construct(const FArguments& InArgs, TSharedPtr<SFilmOv
 		{
 			PrimaryColorTint = OldPrimaryOverlay->GetTint();
 		}
-	
-		OverlayWidgetPinned->SetFilmOverlays(TAttribute<TArray<IFilmOverlay*>>::CreateSP(this, &SFilmOverlayOptions::GetActiveFilmOverlays));
 
 		// CurrentPrimaryOverlay (and therefore Thumbnail) resets for non button version of this widget every time we create it, so let's keep it up to date
 		if (!bIsComboButton)
@@ -792,7 +810,6 @@ FReply SFilmOverlayOptions::SetPrimaryFilmOverlay(FName InName)
 {
 	if (TSharedPtr<SFilmOverlay> OverlayWidgetPinned = OverlayWidget.Pin())
 	{
-		OverlayWidgetPinned->SetFilmOverlays(TAttribute<TArray<IFilmOverlay*>>::CreateSP(this, &SFilmOverlayOptions::GetActiveFilmOverlays));
 		OverlayWidgetPinned->SetPrimaryFilmOverlay(InName);
 	}
 
@@ -848,10 +865,22 @@ TArray<IFilmOverlay*> SFilmOverlayOptions::GetActiveFilmOverlays() const
 
 const FSlateBrush* SFilmOverlayOptions::GetCurrentThumbnail() const
 {
-	FName PrimaryOverlay = GetFilmOverlayWidget()->GetPrimaryFilmOverlay();
+	TSharedPtr<SFilmOverlay> Widget = GetFilmOverlayWidget();
+	if (!Widget.IsValid())
+	{
+		return FLevelSequenceEditorStyle::Get()->GetBrush("FilmOverlay.DefaultThumbnail");
+	}
+
+	FName PrimaryOverlay = Widget->GetPrimaryFilmOverlay();
 	if (!PrimaryOverlay.IsNone())
 	{
-		return UFilmOverlayToolkit::GetPrimaryFilmOverlays()[PrimaryOverlay].Get()->GetThumbnail();
+		if (const TSharedPtr<IFilmOverlay>* Found = UFilmOverlayToolkit::GetPrimaryFilmOverlays().Find(PrimaryOverlay))
+		{
+			if (Found->IsValid())
+			{			
+				return Found->Get()->GetThumbnail();
+			}
+		}
 	}
 
 	return FLevelSequenceEditorStyle::Get()->GetBrush("FilmOverlay.DefaultThumbnail");
@@ -880,67 +909,122 @@ void SFilmOverlayOptions::BindCommands(TSharedRef<FUICommandList> Bindings)
 	TArray<FName> ToggleableOverlayNames;
 	UFilmOverlayToolkit::GetToggleableFilmOverlays().GenerateKeyArray(ToggleableOverlayNames);
 
+	// Use a weak pointer so that if this widget is destroyed (e.g. when the submenu closes
+	// but SCinematicLevelViewport's CommandList still holds these bindings), invoking the
+	// shortcut does not crash via a dangling raw 'this' pointer.
+	TWeakPtr<SFilmOverlayOptions> WeakThis = StaticCastSharedRef<SFilmOverlayOptions>(AsShared());
+
 	Bindings->MapAction(
 		Commands.Disabled,
-		FExecuteAction::CreateLambda([this] { SetPrimaryFilmOverlay(NAME_None); }));
+		FExecuteAction::CreateLambda([WeakThis]
+		{
+			if (TSharedPtr<SFilmOverlayOptions> Pinned = WeakThis.Pin())
+			{
+				Pinned->SetPrimaryFilmOverlay(NAME_None);
+			}
+		}));
 
 	for (int32 OverlayIndex = 0; OverlayIndex < PrimaryOverlayNames.Num(); ++OverlayIndex)
 	{
-		IFilmOverlay& Overlay = *UFilmOverlayToolkit::GetPrimaryFilmOverlays()[PrimaryOverlayNames[OverlayIndex]].Get();
-
 		if (Commands.Grid2x2.Get()->GetCommandName() == PrimaryOverlayNames[OverlayIndex])
 		{
 			Bindings->MapAction(
 				Commands.Grid2x2,
-				FExecuteAction::CreateLambda([this, OverlayName = PrimaryOverlayNames[OverlayIndex]] { SetPrimaryFilmOverlay(OverlayName); }));
+				FExecuteAction::CreateLambda([WeakThis, OverlayName = PrimaryOverlayNames[OverlayIndex]]
+				{
+					if (TSharedPtr<SFilmOverlayOptions> Pinned = WeakThis.Pin())
+					{
+						Pinned->SetPrimaryFilmOverlay(OverlayName);
+					}
+				}));
 		}
 		else if (Commands.Grid3x3.Get()->GetCommandName() == PrimaryOverlayNames[OverlayIndex])
 		{
 			Bindings->MapAction(
 				Commands.Grid3x3,
-				FExecuteAction::CreateLambda([this, OverlayName = PrimaryOverlayNames[OverlayIndex]] { SetPrimaryFilmOverlay(OverlayName); }));
+				FExecuteAction::CreateLambda([WeakThis, OverlayName = PrimaryOverlayNames[OverlayIndex]]
+				{
+					if (TSharedPtr<SFilmOverlayOptions> Pinned = WeakThis.Pin())
+					{
+						Pinned->SetPrimaryFilmOverlay(OverlayName);
+					}
+				}));
 		}
 		else if (Commands.Crosshair.Get()->GetCommandName() == PrimaryOverlayNames[OverlayIndex])
 		{
 			Bindings->MapAction(
 				Commands.Crosshair,
-				FExecuteAction::CreateLambda([this, OverlayName = PrimaryOverlayNames[OverlayIndex]] { SetPrimaryFilmOverlay(OverlayName); }));
+				FExecuteAction::CreateLambda([WeakThis, OverlayName = PrimaryOverlayNames[OverlayIndex]]
+				{
+					if (TSharedPtr<SFilmOverlayOptions> Pinned = WeakThis.Pin())
+					{
+						Pinned->SetPrimaryFilmOverlay(OverlayName);
+					}
+				}));
 		}
 		else if (Commands.Rabatment.Get()->GetCommandName() == PrimaryOverlayNames[OverlayIndex])
 		{
 			Bindings->MapAction(
 				Commands.Rabatment,
-				FExecuteAction::CreateLambda([this, OverlayName = PrimaryOverlayNames[OverlayIndex]] { SetPrimaryFilmOverlay(OverlayName); }));
+				FExecuteAction::CreateLambda([WeakThis, OverlayName = PrimaryOverlayNames[OverlayIndex]]
+				{
+					if (TSharedPtr<SFilmOverlayOptions> Pinned = WeakThis.Pin())
+					{
+						Pinned->SetPrimaryFilmOverlay(OverlayName);
+					}
+				}));
 		}
 	}
 
 	for (int32 OverlayIndex = 0; OverlayIndex < ToggleableOverlayNames.Num(); ++OverlayIndex)
 	{
-		IFilmOverlay& Overlay = *UFilmOverlayToolkit::GetToggleableFilmOverlays()[ToggleableOverlayNames[OverlayIndex]].Get();
-
 		if (Commands.ActionSafe.Get()->GetCommandName() == ToggleableOverlayNames[OverlayIndex])
 		{
 			Bindings->MapAction(
 				Commands.ActionSafe,
-				FExecuteAction::CreateLambda([this, OverlayName = ToggleableOverlayNames[OverlayIndex]] { ToggleFilmOverlay(OverlayName); }));
+				FExecuteAction::CreateLambda([WeakThis, OverlayName = ToggleableOverlayNames[OverlayIndex]]
+				{
+					if (TSharedPtr<SFilmOverlayOptions> Pinned = WeakThis.Pin())
+					{
+						Pinned->ToggleFilmOverlay(OverlayName);
+					}
+				}));
 		}
 		else if (Commands.TitleSafe.Get()->GetCommandName() == ToggleableOverlayNames[OverlayIndex])
 		{
 			Bindings->MapAction(
 				Commands.TitleSafe,
-				FExecuteAction::CreateLambda([this, OverlayName = ToggleableOverlayNames[OverlayIndex]] { ToggleFilmOverlay(OverlayName); }));
+				FExecuteAction::CreateLambda([WeakThis, OverlayName = ToggleableOverlayNames[OverlayIndex]]
+				{
+					if (TSharedPtr<SFilmOverlayOptions> Pinned = WeakThis.Pin())
+					{
+						Pinned->ToggleFilmOverlay(OverlayName);
+					}
+				}));
 		}
 		else if (Commands.CustomSafe.Get()->GetCommandName() == ToggleableOverlayNames[OverlayIndex])
 		{
 			Bindings->MapAction(
 				Commands.CustomSafe,
-				FExecuteAction::CreateLambda([this, OverlayName = ToggleableOverlayNames[OverlayIndex]] { ToggleFilmOverlay(OverlayName); }));
+				FExecuteAction::CreateLambda([WeakThis, OverlayName = ToggleableOverlayNames[OverlayIndex]]
+				{
+					if (TSharedPtr<SFilmOverlayOptions> Pinned = WeakThis.Pin())
+					{
+						Pinned->ToggleFilmOverlay(OverlayName);
+					}
+				}));
 		}
 		else if (Commands.Letterbox.Get()->GetCommandName() == ToggleableOverlayNames[OverlayIndex])
 		{
 			Bindings->MapAction(
 				Commands.Letterbox,
-				FExecuteAction::CreateLambda([this, OverlayName = ToggleableOverlayNames[OverlayIndex]] { ToggleFilmOverlay(OverlayName); }));
+				FExecuteAction::CreateLambda([WeakThis, OverlayName = ToggleableOverlayNames[OverlayIndex]]
+				{
+					if (TSharedPtr<SFilmOverlayOptions> Pinned = WeakThis.Pin())
+					{
+						Pinned->ToggleFilmOverlay(OverlayName);
+					}
+				}));
 		}
 	}
 }
diff --git a/Engine/Plugins/MovieScene/LevelSequenceEditor/Source/LevelSequenceEditor/Private/CinematicViewport/FilmOverlays.h b/Engine/Plugins/MovieScene/LevelSequenceEditor/Source/LevelSequenceEditor/Private/CinematicViewport/FilmOverlays.h
--- a/Engine/Plugins/MovieScene/LevelSequenceEditor/Source/LevelSequenceEditor/Private/CinematicViewport/FilmOverlays.h
+++ b/Engine/Plugins/MovieScene/LevelSequenceEditor/Source/LevelSequenceEditor/Private/CinematicViewport/FilmOverlays.h
@@ -74,6 +74,9 @@ public:
 	/** Bind commands for the overlays */
 	void BindCommands(TSharedRef<FUICommandList>);
 
+	/** Register the default primary and toggleable overlays into UFilmOverlayToolkit (idempotent). */
+	static void RegisterDefaultOverlays();
+
 private:
 
 	/** Generate menu content for the combo button */
diff --git a/Engine/Plugins/MovieScene/LevelSequenceEditor/Source/LevelSequenceEditor/Private/CinematicViewport/SCinematicLevelViewport.cpp b/Engine/Plugins/MovieScene/LevelSequenceEditor/Source/LevelSequenceEditor/Private/CinematicViewport/SCinematicLevelViewport.cpp
--- a/Engine/Plugins/MovieScene/LevelSequenceEditor/Source/LevelSequenceEditor/Private/CinematicViewport/SCinematicLevelViewport.cpp
+++ b/Engine/Plugins/MovieScene/LevelSequenceEditor/Source/LevelSequenceEditor/Private/CinematicViewport/SCinematicLevelViewport.cpp
@@ -22,6 +22,7 @@
 #include "CinematicViewport/CinematicViewportCommands.h"
 #include "CinematicViewport/SCinematicTransportRange.h"
 #include "CinematicViewport/FilmOverlays.h"
+#include "FilmOverlayToolkit.h"
 #include "Widgets/Layout/SWidgetSwitcher.h"
 #include "Fonts/FontMeasure.h"
 #include "AssetEditorViewportLayout.h"
@@ -94,11 +95,6 @@ void AddOverlayOptionsEntry(UToolMenu* InMenu)
 					TSharedRef<SFilmOverlayOptions> FilmOverlayOptions =
 						SNew(SFilmOverlayOptions, CinematicLevelViewportPinned->GetOverlayWidget()).IsComboButton(false);
 
-					if (TSharedPtr<FUICommandList> CommandList = CinematicLevelViewportPinned->GetCommandList())
-					{
-						FilmOverlayOptions->BindCommands(CommandList.ToSharedRef());
-					}
-
 					Submenu->AddMenuEntry(
 						"FilmOverlay", FToolMenuEntry::InitWidget("FilmOverlay", FilmOverlayOptions, FText())
 					);
@@ -355,6 +351,49 @@ void SCinematicLevelViewport::Construct(const FArguments& InArgs)
 
 	OverlayWidget = SNew(SFilmOverlay).Visibility(EVisibility::HitTestInvisible);
 
+	// Register default overlays and bind overlay shortcuts to persistent objects so the shortcuts
+	// continue to work after toolbar refreshes (e.g. after Save All) destroy transient SFilmOverlayOptions.
+	SFilmOverlayOptions::RegisterDefaultOverlays();
+	{
+		const FCinematicViewportCommands& Commands = FCinematicViewportCommands::Get();
+		TWeakPtr<SFilmOverlay> WeakOverlay = OverlayWidget;
+
+		auto MakePrimaryAction = [WeakOverlay](FName OverlayName) -> FExecuteAction
+		{
+			return FExecuteAction::CreateLambda([WeakOverlay, OverlayName]
+			{
+				if (TSharedPtr<SFilmOverlay> Pinned = WeakOverlay.Pin())
+				{
+					Pinned->SetPrimaryFilmOverlay(OverlayName);
+				}
+			});
+		};
+
+		auto MakeToggleAction = [](FName OverlayName) -> FExecuteAction
+		{
+			return FExecuteAction::CreateLambda([OverlayName]
+			{
+				if (const TSharedPtr<IFilmOverlay>* Overlay = UFilmOverlayToolkit::GetToggleableFilmOverlays().Find(OverlayName))
+				{
+					if (Overlay->IsValid())
+					{
+						(*Overlay)->SetEnabled(!(*Overlay)->IsEnabled());
+					}
+				}
+			});
+		};
+
+		CommandList->MapAction(Commands.Disabled,    MakePrimaryAction(NAME_None));
+		CommandList->MapAction(Commands.Grid3x3,     MakePrimaryAction(FName("Grid3x3")));
+		CommandList->MapAction(Commands.Grid2x2,     MakePrimaryAction(FName("Grid2x2")));
+		CommandList->MapAction(Commands.Crosshair,   MakePrimaryAction(FName("Crosshair")));
+		CommandList->MapAction(Commands.Rabatment,   MakePrimaryAction(FName("Rabatment")));
+		CommandList->MapAction(Commands.ActionSafe,  MakeToggleAction(FName("ActionSafe")));
+		CommandList->MapAction(Commands.TitleSafe,   MakeToggleAction(FName("TitleSafe")));
+		CommandList->MapAction(Commands.CustomSafe,  MakeToggleAction(FName("CustomSafe")));
+		CommandList->MapAction(Commands.Letterbox,   MakeToggleAction(FName("LetterBox")));
+	}
+
 	DecoratedTransportControls = SNew(SHorizontalBox)
 		+ SHorizontalBox::Slot()
 		[
